Live Link VCAM for iOS initial check-in.

diff --git a/Engine/Extras/VirtualProduction/LiveLinkVCAM/Tentacle.xcframework/ios-arm64_armv7/Tentacle.framework/Headers/Tentacle.h b/Engine/Extras/VirtualProduction/LiveLinkVCAM/Tentacle.xcframework/ios-arm64_armv7/Tentacle.framework/Headers/Tentacle.h
new file mode 100644
--- /dev/null
+++ b/Engine/Extras/VirtualProduction/LiveLinkVCAM/Tentacle.xcframework/ios-arm64_armv7/Tentacle.framework/Headers/Tentacle.h
@@ -0,0 +1,52 @@
+//
+//  Tentacle.h
+//  TentacleFramework
+//
+//  Copyright 2019 Tentacle Sync GmbH
+//
+//  Redistribution and use in source and binary forms, with or without
+//  modification, are permitted provided that the following conditions are met:
+//
+//  1. Redistributions of source code must retain the above copyright notice,
+//     this list of conditions and the following disclaimer.
+//
+//  2. Redistributions in binary form must reproduce the above copyright notice,
+//     this list of conditions and the following disclaimer in the
+//     documentation and/or other materials provided with the distribution.
+//
+//  3. Neither the name of the copyright holder nor the names of its
+//     contributors may be used to endorse or promote products derived from this
+//     software without specific prior written permission.
+//
+//  THIS SOFTWARE IS PROVIDED BY THE COPYRIGHT HOLDERS AND CONTRIBUTORS "AS IS"
+//  AND ANY EXPRESS OR IMPLIED WARRANTIES, INCLUDING, BUT NOT LIMITED TO, THE
+//  IMPLIED WARRANTIES OF MERCHANTABILITY AND FITNESS FOR A PARTICULAR PURPOSE
+//  ARE DISCLAIMED. IN NO EVENT SHALL THE COPYRIGHT HOLDER OR CONTRIBUTORS BE
+//  LIABLE FOR ANY DIRECT, INDIRECT, INCIDENTAL, SPECIAL, EXEMPLARY, OR
+//  CONSEQUENTIAL DAMAGES (INCLUDING, BUT NOT LIMITED TO, PROCUREMENT OF
+//  SUBSTITUTE GOODS OR SERVICES; LOSS OF USE, DATA, OR PROFITS; OR BUSINESS
+//  INTERRUPTION) HOWEVER CAUSED AND ON ANY THEORY OF LIABILITY, WHETHER IN
+//  CONTRACT, STRICT LIABILITY, OR TORT (INCLUDING NEGLIGENCE OR OTHERWISE)
+//  ARISING IN ANY WAY OUT OF THE USE OF THIS SOFTWARE, EVEN IF ADVISED OF THE
+//  POSSIBILITY OF SUCH DAMAGE.
+//
+
+#import <Foundation/Foundation.h>
+
+//! Project version number for TentacleFramework.
+FOUNDATION_EXPORT double TentacleFrameworkVersionNumber;
+
+//! Project version string for TentacleFramework.
+FOUNDATION_EXPORT const unsigned char TentacleFrameworkVersionString[];
+
+// In this header, you should import all the public headers of your framework using statements like #import <Tentacle/PublicHeader.h>
+#import <Tentacle/TentacleTimecode.h>
+#import <Tentacle/TentacleAdvertisement.h>
+#import <Tentacle/TentacleDevice.h>
+#import <Tentacle/TentacleDeviceCache.h>
+#import <Tentacle/TentacleBluetoothController.h>
+#import <Tentacle/TentacleLog.h>
+
+#if TARGET_OS_IOS
+#import <Tentacle/TentacleIcon.h>
+#endif
diff --git a/Engine/Extras/VirtualProduction/LiveLinkVCAM/Tentacle.xcframework/ios-arm64_armv7/Tentacle.framework/Headers/TentacleAdvertisement.h b/Engine/Extras/VirtualProduction/LiveLinkVCAM/Tentacle.xcframework/ios-arm64_armv7/Tentacle.framework/Headers/TentacleAdvertisement.h
new file mode 100644
--- /dev/null
+++ b/Engine/Extras/VirtualProduction/LiveLinkVCAM/Tentacle.xcframework/ios-arm64_armv7/Tentacle.framework/Headers/TentacleAdvertisement.h
@@ -0,0 +1,223 @@
+//
+//  tentacleAdvertisement.h
+//  Tentacle-SDK
+//
+//  Copyright 2019 Tentacle Sync GmbH
+//
+//  Redistribution and use in source and binary forms, with or without
+//  modification, are permitted provided that the following conditions are met:
+//
+//  1. Redistributions of source code must retain the above copyright notice,
+//     this list of conditions and the following disclaimer.
+//
+//  2. Redistributions in binary form must reproduce the above copyright notice,
+//     this list of conditions and the following disclaimer in the
+//     documentation and/or other materials provided with the distribution.
+//
+//  3. Neither the name of the copyright holder nor the names of its
+//     contributors may be used to endorse or promote products derived from this
+//     software without specific prior written permission.
+//
+//  THIS SOFTWARE IS PROVIDED BY THE COPYRIGHT HOLDERS AND CONTRIBUTORS "AS IS"
+//  AND ANY EXPRESS OR IMPLIED WARRANTIES, INCLUDING, BUT NOT LIMITED TO, THE
+//  IMPLIED WARRANTIES OF MERCHANTABILITY AND FITNESS FOR A PARTICULAR PURPOSE
+//  ARE DISCLAIMED. IN NO EVENT SHALL THE COPYRIGHT HOLDER OR CONTRIBUTORS BE
+//  LIABLE FOR ANY DIRECT, INDIRECT, INCIDENTAL, SPECIAL, EXEMPLARY, OR
+//  CONSEQUENTIAL DAMAGES (INCLUDING, BUT NOT LIMITED TO, PROCUREMENT OF
+//  SUBSTITUTE GOODS OR SERVICES; LOSS OF USE, DATA, OR PROFITS; OR BUSINESS
+//  INTERRUPTION) HOWEVER CAUSED AND ON ANY THEORY OF LIABILITY, WHETHER IN
+//  CONTRACT, STRICT LIABILITY, OR TORT (INCLUDING NEGLIGENCE OR OTHERWISE)
+//  ARISING IN ANY WAY OUT OF THE USE OF THIS SOFTWARE, EVEN IF ADVISED OF THE
+//  POSSIBILITY OF SUCH DAMAGE.
+//
+
+/**
+ @file
+ @brief Advertisement data.
+ 
+ Contains data of a received Bluetooth advertisement packet.
+ */
+
+#ifndef tentacleAdvertisement_h
+#define tentacleAdvertisement_h
+
+#ifdef __cplusplus
+extern "C" {
+#endif
+
+#include <stdio.h>
+#include <stdbool.h>
+#include <string.h>
+#include <Tentacle/TentacleEnum.h>
+#include <Tentacle/TentacleTimecode.h>
+
+/// Length of string returned by `tentacleAdvertisementString`.
+#define TENTACLE_ADVERTISEMENT_STRING_LENGTH 104
+    
+/// Length of unique identifier representing a discovered device, see `tentacleAdvertisement.identifier`.
+#define TENTACLE_ADVERTISEMENT_ID_LENGTH_MAX 37
+    
+/// Maximum length of a device name, see `tentacleAdvertisement.name`.
+#define TENTACLE_ADVERTISEMENT_NAME_LENGTH_MAX 16
+
+/**
+ Contains data of a received Bluetooth advertisement packet.
+ */
+typedef struct TentacleAdvertisement {
+    /// Tentacle Bluetooth protocol version. Currently supported version is 1.
+    uint8_t version;
+    
+    /// True if device is able to connect via Bluetooth.
+    bool connectable;
+    
+    /// True if a mobile device can authenticate to the device.
+    bool linkMode;
+    
+    /// Device type. Currently, the only valid value is 0 for Tentacle Sync-E.
+    uint8_t deviceType;
+    
+    /// State the device is in.
+    enum tentacleDeviceMode deviceMode;
+    
+    /// Audio volume when generating audio timecode.
+    enum TentacleAudioLevel audioLevel;
+    
+    /// True if device is generating NTSC timecode.
+    bool ntsc;
+    
+    /// True if device is generating drop frame timecode.
+    bool dropFrame;
+    
+    /// \brief True if device is in green mode.
+    ///
+    /// Green mode means that the timecode has been configured, the device is
+    /// generating audio timecode and other devices can be synced to this device
+    /// using a cable.
+    bool greenMode;
+    
+    /// Hour, minute, second, frame, microseconds and reception timestamp
+    TentacleTimecode timecode;
+    
+    /// Battery level of device in percent.
+    uint8_t battery;
+    
+    /// True if device is currently charging.
+    bool charging;
+    
+    /// Icon index the device is set to.
+    uint8_t icon;
+    
+    /// Received Signal Strength Indication
+    int8_t rssi;
+    
+    /// \brief Unique identifier string representing the device this
+    /// advertisement was sent from.
+    ///
+    /// On iOS this is the string value of a [CBUUID](https://developer.apple.com/documentation/corebluetooth/cbuuid) object.
+    ///
+    /// On Android this is the hardware address as returned by [BluetoothDevice.getAddress()](https://developer.android.com/reference/android/bluetooth/BluetoothDevice.html#getAddress()),
+    ///
+    /// On Nordic nRF5 SDK this is the hardware address as passed in  [bleGapAddr_t](https://www.nordicsemi.com/DocLib/Content/SoftDeviceAPIDoc/S132/v6-1-0/structble_Gap_Addr__t).
+    char identifier[TENTACLE_ADVERTISEMENT_ID_LENGTH_MAX];
+    
+    /// Length of identifier string, has to be less than `tentacleAdvertisementIdentifierLength`.
+    size_t identifierLength;
+    
+    /// Name of the device this advertisement was sent from.
+    char name[TENTACLE_ADVERTISEMENT_NAME_LENGTH_MAX];
+    
+    /// Length of name string, has to be less than `tentacleNameLength`.
+    size_t nameLength;
+    
+    /// @brief Current frame rate as integer value, ignoring NTSC flag.
+    ///
+    /// To calculate the actual frame rate, you have to take the NTSC flag into
+    /// account or use `tentacleAdvertisementGetFrameRate()`.
+    int frameRateInt;
+    
+    /// @brief True if data in this struct is considered to describe a valid timecode.
+    ///
+    /// Value is set by `tentacleAdvertisementInit()`.
+    bool valid;
+} TentacleAdvertisement;
+    
+/**
+ Returns an advertisement struct with default values.
+ 
+ @return TentacleAdvertisement struct with default values.
+ */
+TentacleAdvertisement TentacleAdvertisementMake(void);
+
+/**
+ Initializes a `tentacleAdvertisement` from received Bluetooth advertisement
+ packet data.
+
+ If the passed data is considered to be valid, `advertisement.valid` is set to true.
+
+ @param advertisementData Byte array containing advertisement packet data.
+ @param advertisementDataLen Length of advertisement data array.
+ @param rssi Received Signal Strength Indication.
+ @param receivedTimestamp Precise timestamp of when this packet was received,
+        in seconds.
+ @param identifierBuffer Unique identifier string representing the device this
+        advertisement was sent from.
+ @param identifierBufferLength Length of identifier string, has to be less than
+        `tentacleAdvertisementIdentifierLengthMax`.
+ @param nameBuffer Name of the device this advertisement was sent from.
+ @param nameBufferSize Length of name string, has to be less than
+        `tentacleNameLength`.
+ @return Advertisement struct containing received packet data.
+ */
+TentacleAdvertisement TentacleAdvertisementInit(const unsigned char *advertisementData, size_t advertisementDataLen, int8_t rssi, double receivedTimestamp, const char *identifierBuffer, size_t identifierBufferLength, const char *nameBuffer, size_t nameBufferSize);
+
+/**
+ Check if two frame rates are equal.
+ 
+ @param advertisement1 First frame rate to compare.
+ @param advertisement2 Second frame rate to compare.
+ @return True if frame rate are equal, false otherwise.
+ */
+bool TentacleAdvertisementFrameRateEqual(const TentacleAdvertisement *advertisement1, const TentacleAdvertisement *advertisement2);
+
+/**
+ Frame rate, depending on value of NTSC flag. Also see
+ `tentacleAdvertisement.frameRateInt` for frame rate ignoring the NTSC flag.
+
+ @param advertisement advertisement to get the frame rate of.
+ @return frame rate depending on NTSC flag.
+ */
+double TentacleAdvertisementGetFrameRate(const TentacleAdvertisement *advertisement);
+
+/**
+ String value of advertisement, containing frame rate, battery state and more.
+
+ @param advertisement Get string value of this advertisement.
+ @param advertisementString Resulting string value. Has to be of length
+        `TENTACLE_ADVERTISEMENT_STRING_LENGTH`, includes terminating null
+        character.
+ @return Number of characters of `advertisementString`, not including the
+         terminating null character.
+ */
+int TentacleAdvertisementString(const TentacleAdvertisement *advertisement, char *advertisementString);
+    
+/**
+ Name of the device this advertisement was received from.
+
+ @param advertisement advertisement to get the name of.
+ @param nameString Resulting name string
+ @return Length of returned string.
+ */
+int TentacleAdvertisementNameString(const TentacleAdvertisement *advertisement, char *nameString);
+
+/**
+ Print string representation of advertisement to stdout.
+
+ @param advertisement Advertisement struct to print.
+ */
+void TentacleAdvertisementPrint(const TentacleAdvertisement *advertisement);
+
+#ifdef __cplusplus
+}
+#endif
+
+#endif /* tentacleAdvertisement_h */
diff --git a/Engine/Extras/VirtualProduction/LiveLinkVCAM/Tentacle.xcframework/ios-arm64_armv7/Tentacle.framework/Headers/TentacleBluetoothController.h b/Engine/Extras/VirtualProduction/LiveLinkVCAM/Tentacle.xcframework/ios-arm64_armv7/Tentacle.framework/Headers/TentacleBluetoothController.h
new file mode 100644
--- /dev/null
+++ b/Engine/Extras/VirtualProduction/LiveLinkVCAM/Tentacle.xcframework/ios-arm64_armv7/Tentacle.framework/Headers/TentacleBluetoothController.h
@@ -0,0 +1,117 @@
+//
+//  TentacleBluetoothController.h
+//  Tentacle-SDK
+//
+//  Copyright 2019 Tentacle Sync GmbH
+//
+//  Redistribution and use in source and binary forms, with or without
+//  modification, are permitted provided that the following conditions are met:
+//
+//  1. Redistributions of source code must retain the above copyright notice,
+//     this list of conditions and the following disclaimer.
+//
+//  2. Redistributions in binary form must reproduce the above copyright notice,
+//     this list of conditions and the following disclaimer in the
+//     documentation and/or other materials provided with the distribution.
+//
+//  3. Neither the name of the copyright holder nor the names of its
+//     contributors may be used to endorse or promote products derived from this
+//     software without specific prior written permission.
+//
+//  THIS SOFTWARE IS PROVIDED BY THE COPYRIGHT HOLDERS AND CONTRIBUTORS "AS IS"
+//  AND ANY EXPRESS OR IMPLIED WARRANTIES, INCLUDING, BUT NOT LIMITED TO, THE
+//  IMPLIED WARRANTIES OF MERCHANTABILITY AND FITNESS FOR A PARTICULAR PURPOSE
+//  ARE DISCLAIMED. IN NO EVENT SHALL THE COPYRIGHT HOLDER OR CONTRIBUTORS BE
+//  LIABLE FOR ANY DIRECT, INDIRECT, INCIDENTAL, SPECIAL, EXEMPLARY, OR
+//  CONSEQUENTIAL DAMAGES (INCLUDING, BUT NOT LIMITED TO, PROCUREMENT OF
+//  SUBSTITUTE GOODS OR SERVICES; LOSS OF USE, DATA, OR PROFITS; OR BUSINESS
+//  INTERRUPTION) HOWEVER CAUSED AND ON ANY THEORY OF LIABILITY, WHETHER IN
+//  CONTRACT, STRICT LIABILITY, OR TORT (INCLUDING NEGLIGENCE OR OTHERWISE)
+//  ARISING IN ANY WAY OUT OF THE USE OF THIS SOFTWARE, EVEN IF ADVISED OF THE
+//  POSSIBILITY OF SUCH DAMAGE.
+//
+
+@import CoreBluetooth;
+#include <Tentacle/TentacleDevice.h>
+
+NS_ASSUME_NONNULL_BEGIN
+
+/// State of the Bluetooth controller. Uses the same values as <code>CBManagerState</code>.
+/// See https://developer.apple.com/documentation/corebluetooth/cbmanagerstate
+typedef NS_ENUM(NSUInteger, TentacleBluetoothState) {
+    TentacleBluetoothStateUnknown = 0,
+    TentacleBluetoothStateResetting = 1,
+    TentacleBluetoothStateUnsupported = 2,
+    TentacleBluetoothStateUnauthorized = 3,
+    TentacleBluetoothStatePoweredOff = 4,
+    TentacleBluetoothStatePoweredOn = 5,
+};
+
+/// The delegate of a <code>TentacleBluetoothController</code> object must adopt
+/// this protocol. The delegate uses this protocol’s methods to notify about the
+/// discovery, state changes and disappearance of Bluetooth devices.
+@protocol TentacleBluetoothControllerDelegate <NSObject>
+@optional
+/// Called if state of <code>TentacleBluetoothController</code> changed.
+/// \param state New state the Bluetooth controller is in.
+- (void)didUpdateToState:(enum TentacleBluetoothState)state;
+
+/// Called if scanning for Bluetooth devices has started.
+- (void)didStartScanning;
+
+@required
+/// Called each time advertisment data has been received from a Tentacle.
+/// Advertisment broadcast rate is multiple times per second, thus this may
+/// be called frequently.
+/// \param deviceIndex Device with this index has been updated
+- (void)didReceiveAdvertismentForDeviceIndex:(NSInteger)deviceIndex;
+@end
+
+
+@interface TentacleBluetoothController : NSObject <CBCentralManagerDelegate>
+/// Dictionary of all devices which have been seen.
+/// Peripherals are only added to this array but never removed.
+@property (nonatomic, strong, readonly) NSMutableDictionary<NSString *, CBPeripheral *> * _Nonnull discoveredPeripherals;
+    
+/// <code>CBCentralManager</code> instance used to scan for devices.
+@property (nonatomic, strong, readonly) CBCentralManager * _Nonnull centralManager;
+
+/// UUID of the Tentacle Bluetooth service.
+@property (nonatomic, strong, readonly) CBUUID * _Nonnull tentacleServiceUUID;
+    
+/// Delegate to receive the Bluetooth events.
+@property (nonatomic, weak) id <TentacleBluetoothControllerDelegate> delegate;
+
+/// Resets the central manager
+- (void)reset;
+
+/// Starts scanning for Bluetooth Tentacles. Scanning is started with the
+/// <code>CBCentralManagerScanOptionAllowDuplicatesKey</code> option set.
+- (void)startScanning;
+
+/// Stops scanning for Bluetooth Tentacles.
+- (void)stopScanning;
+
+/// Returns a Tentacle device with a specific index from device cache.
+/// \param deviceIndex Get device with this index in device cache.
+///
+/// returns:
+/// Device with specified index.
+- (TentacleDevice)retrieveDeviceWithIndex:(int)deviceIndex;
+
+/// Precise timestamp in seconds.
+///
+/// Uses the `mach_absolute_time()` mach kernel syscall.
+///
+/// See section
+/// "Converting Mach absolute time to nanoseconds using mach_timebase_info" in
+/// Apple's technical Q&A [Mach Absolute Time Units](https://developer.apple.com/library/content/qa/qa1398/_index.html)
+/// and the [mach_time.h]([REDACTED_PROFILE_URL]/darwin-xnu/blob/master/osfmk/mach/mach_time.h)
+/// kernel sources for the [mach_timebase_info]([REDACTED_PROFILE_URL]/darwin-xnu/blob/xnu-4570.1.46/osfmk/mach/mach_time.h#L36) struct.
+///
+/// returns:
+/// Precise timestamp in seconds.
++ (double)timestamp;
+@end
+
+NS_ASSUME_NONNULL_END
diff --git a/Engine/Extras/VirtualProduction/LiveLinkVCAM/Tentacle.xcframework/ios-arm64_armv7/Tentacle.framework/Headers/TentacleDevice.h b/Engine/Extras/VirtualProduction/LiveLinkVCAM/Tentacle.xcframework/ios-arm64_armv7/Tentacle.framework/Headers/TentacleDevice.h
new file mode 100644
--- /dev/null
+++ b/Engine/Extras/VirtualProduction/LiveLinkVCAM/Tentacle.xcframework/ios-arm64_armv7/Tentacle.framework/Headers/TentacleDevice.h
@@ -0,0 +1,163 @@
+//
+//  tentacleDevice.h
+//  Tentacle-SDK
+//
+//  Copyright 2019 Tentacle Sync GmbH
+//
+//  Redistribution and use in source and binary forms, with or without
+//  modification, are permitted provided that the following conditions are met:
+//
+//  1. Redistributions of source code must retain the above copyright notice,
+//     this list of conditions and the following disclaimer.
+//
+//  2. Redistributions in binary form must reproduce the above copyright notice,
+//     this list of conditions and the following disclaimer in the
+//     documentation and/or other materials provided with the distribution.
+//
+//  3. Neither the name of the copyright holder nor the names of its
+//     contributors may be used to endorse or promote products derived from this
+//     software without specific prior written permission.
+//
+//  THIS SOFTWARE IS PROVIDED BY THE COPYRIGHT HOLDERS AND CONTRIBUTORS "AS IS"
+//  AND ANY EXPRESS OR IMPLIED WARRANTIES, INCLUDING, BUT NOT LIMITED TO, THE
+//  IMPLIED WARRANTIES OF MERCHANTABILITY AND FITNESS FOR A PARTICULAR PURPOSE
+//  ARE DISCLAIMED. IN NO EVENT SHALL THE COPYRIGHT HOLDER OR CONTRIBUTORS BE
+//  LIABLE FOR ANY DIRECT, INDIRECT, INCIDENTAL, SPECIAL, EXEMPLARY, OR
+//  CONSEQUENTIAL DAMAGES (INCLUDING, BUT NOT LIMITED TO, PROCUREMENT OF
+//  SUBSTITUTE GOODS OR SERVICES; LOSS OF USE, DATA, OR PROFITS; OR BUSINESS
+//  INTERRUPTION) HOWEVER CAUSED AND ON ANY THEORY OF LIABILITY, WHETHER IN
+//  CONTRACT, STRICT LIABILITY, OR TORT (INCLUDING NEGLIGENCE OR OTHERWISE)
+//  ARISING IN ANY WAY OUT OF THE USE OF THIS SOFTWARE, EVEN IF ADVISED OF THE
+//  POSSIBILITY OF SUCH DAMAGE.
+//
+
+/**
+ @file
+ @brief Tentacle device.
+ 
+ Cached data of a discovered device.
+ */
+
+#ifndef tentacledevice_h
+#define tentacledevice_h
+
+#ifdef __cplusplus
+extern "C" {
+#endif
+    
+#include <stdio.h>
+#include <Tentacle/TentacleTimecode.h>
+#include <Tentacle/TentacleAdvertisement.h>
+    
+/// Maximum of timecode values to cache for each device
+#define TENTACLE_DEVICE_TIMECODE_CACHE_SIZE_MAX 10
+
+/// Maximum of received signal strength indicator (RSSI) values for each device
+#define TENTACLE_DEVICE_RSSI_CACHE_SIZE_MAX 5
+
+/// If a device has not been seen after the defined time span it is considered
+/// disappeared.
+#define TENTACLE_DEVICE_DISAPPEARED_THRESHOLD 10.0
+
+/// If a device has not been seen after the defined time span it is considered
+/// unavailable.
+#define TENTACLE_DEVICE_UNAVAILABLE_THRESHOLD 600.0
+
+/**
+ Cached data of a discovered device.
+ */
+typedef struct TentacleDevice {
+    /// Last received advertisement for this device
+    struct TentacleAdvertisement advertisement;
+    
+    /// Ring buffer of cached timecodes
+    TentacleTimecode timecodeCache[TENTACLE_DEVICE_TIMECODE_CACHE_SIZE_MAX];
+
+    /// Current index of timecode in ring buffer
+    size_t timecodeCacheIndex;
+    
+    /// Current size of timecode ring buffer
+    size_t timecodeCacheSize;
+    
+    /// Averaged timecode of all timecodes from timecode cache.
+    TentacleTimecode timecode;
+    
+    /// Ring buffer of received signal strength indication (RSSI) values
+    int rssiCache[TENTACLE_DEVICE_RSSI_CACHE_SIZE_MAX];
+    
+    /// Next index in RSSI ring buffer to use
+    size_t rssiCacheNextIndex;
+    
+    /// @private Used to calculate precise timecode
+    int startPhaseCounter;
+    
+    /// @private Used to calculate precise timecode
+    int asyncCounter;
+    
+    /// @private Used to calculate precise timecode
+    int severeAsyncCounter;
+} TentacleDevice;
+
+/**
+ Returns a device struct with default values.
+ 
+ @return TentacleDevice struct with default values.
+ */
+TentacleDevice TentacleDeviceMake(void);
+    
+/**
+ Update precise, calculated timecode for this device.
+ 
+ @param device Update cached timecodes of this device.
+ @param advertisementNew Add timecode from this newly received advertisement
+ to cache.
+ @return Difference between elapsed time according to received timecodes and
+ local device time, in milliseconds.
+ */
+double TentacleDeviceUpdateTimecode(TentacleDevice *device, const TentacleAdvertisement *advertisementNew);
+
+/**
+ If a device has not been seen after `TENTACLE_DEVICE_DISAPPEARED_THRESHOLD` seconds
+ it is considered disappeared.
+ 
+ @param device Device to get disappeared status for.
+ @param atTimestamp Get disappeared status for this timestamp, i.e. the current
+ time since boot in seconds (see TentacleTimecode.receivedTimestamp for specific timestamp-source).
+ @return True if device has not been seen for `TENTACLE_DEVICE_DISAPPEARED_THRESHOLD`
+ seconds.
+ */
+bool TentacleDeviceIsDisappeared(TentacleDevice *device, double atTimestamp);
+
+/**
+ If a device has not been seen after `TENTACLE_DEVICE_UNAVAILABLE_THRESHOLD` seconds
+ it is considered unavailable.
+
+ @param device Device to get unavailable status for.
+ @param atTimestamp Get unavailable status for this timestamp, i.e. the current
+ time since boot in seconds (see TentacleTimecode.receivedTimestamp for specific timestamp-source).
+ @return True if device has not been seen for `TENTACLE_DEVICE_UNAVAILABLE_THRESHOLD`
+ seconds.
+ */
+bool TentacleDeviceIsUnavailable(TentacleDevice *device, double atTimestamp);
+
+/**
+ Add a received signal strength indicator (RSSI) value to the ring buffer.
+
+ @param device Add to ring buffer of this device.
+ @param rssi Add this RSSI value to ring buffer.
+ */
+void TentacleDeviceAddRssi(TentacleDevice *device, int rssi);
+
+/**
+ Averaged received signal strength indicator (RSSI) value.
+ 
+ @param device Get averaged RSSI value for this device.
+ @return Averaged RSSI value
+ */
+int TentacleDeviceGetRssi(TentacleDevice *device);
+    
+#ifdef __cplusplus
+}
+#endif
+
+#endif /* tentacledevice_h */
diff --git a/Engine/Extras/VirtualProduction/LiveLinkVCAM/Tentacle.xcframework/ios-arm64_armv7/Tentacle.framework/Headers/TentacleDeviceCache.h b/Engine/Extras/VirtualProduction/LiveLinkVCAM/Tentacle.xcframework/ios-arm64_armv7/Tentacle.framework/Headers/TentacleDeviceCache.h
new file mode 100644
--- /dev/null
+++ b/Engine/Extras/VirtualProduction/LiveLinkVCAM/Tentacle.xcframework/ios-arm64_armv7/Tentacle.framework/Headers/TentacleDeviceCache.h
@@ -0,0 +1,136 @@
+//
+//  tentacleDeviceCache.h
+//  Tentacle-SDK
+//
+//  Copyright 2019 Tentacle Sync GmbH
+//
+//  Redistribution and use in source and binary forms, with or without
+//  modification, are permitted provided that the following conditions are met:
+//
+//  1. Redistributions of source code must retain the above copyright notice,
+//     this list of conditions and the following disclaimer.
+//
+//  2. Redistributions in binary form must reproduce the above copyright notice,
+//     this list of conditions and the following disclaimer in the
+//     documentation and/or other materials provided with the distribution.
+//
+//  3. Neither the name of the copyright holder nor the names of its
+//     contributors may be used to endorse or promote products derived from this
+//     software without specific prior written permission.
+//
+//  THIS SOFTWARE IS PROVIDED BY THE COPYRIGHT HOLDERS AND CONTRIBUTORS "AS IS"
+//  AND ANY EXPRESS OR IMPLIED WARRANTIES, INCLUDING, BUT NOT LIMITED TO, THE
+//  IMPLIED WARRANTIES OF MERCHANTABILITY AND FITNESS FOR A PARTICULAR PURPOSE
+//  ARE DISCLAIMED. IN NO EVENT SHALL THE COPYRIGHT HOLDER OR CONTRIBUTORS BE
+//  LIABLE FOR ANY DIRECT, INDIRECT, INCIDENTAL, SPECIAL, EXEMPLARY, OR
+//  CONSEQUENTIAL DAMAGES (INCLUDING, BUT NOT LIMITED TO, PROCUREMENT OF
+//  SUBSTITUTE GOODS OR SERVICES; LOSS OF USE, DATA, OR PROFITS; OR BUSINESS
+//  INTERRUPTION) HOWEVER CAUSED AND ON ANY THEORY OF LIABILITY, WHETHER IN
+//  CONTRACT, STRICT LIABILITY, OR TORT (INCLUDING NEGLIGENCE OR OTHERWISE)
+//  ARISING IN ANY WAY OUT OF THE USE OF THIS SOFTWARE, EVEN IF ADVISED OF THE
+//  POSSIBILITY OF SUCH DAMAGE.
+//
+
+/**
+ @file
+ @brief Discovered devices.
+ */
+
+#ifndef tentacleDeviceCache_h
+#define tentacleDeviceCache_h
+
+#ifdef __cplusplus
+extern "C" {
+#endif
+
+#include <stdio.h>
+#include <Tentacle/TentacleAdvertisement.h>
+#include <Tentacle/TentacleDevice.h>
+
+/**
+ Process a received Bluetooth advertisement packet.
+
+ @param advertisement Advertisement data to be processed.
+ @return Index of updated device in device cache.
+ */
+size_t TentacleDeviceCacheProcess(const TentacleAdvertisement *advertisement);
+    
+/**
+ Get current size of device cache.
+ 
+ @return Size of device cache.
+ */
+size_t TentacleDeviceCacheGetSize(void);
+    
+/**
+ Get current size of timecode cache for a specific device index.
+ 
+ @param deviceIndex Get cache size for device with this index.
+ @return Timecode cache size.
+ */
+size_t TentacleDeviceCacheGetTimecodeCacheSize(int deviceIndex);
+
+/**
+ Get cached device at specific index.
+ 
+ @param deviceIndex Index of device to get.
+ @return Device struct at specific index.
+ */
+TentacleDevice TentacleDeviceCacheGetDevice(int deviceIndex);
+
+/**
+ Get cached advertisement for device at specific index.
+
+ @param deviceIndex Index of device to get advertisement for.
+ @return Cached advertisement struct.
+ */
+TentacleAdvertisement TentacleDeviceCacheGetAdvertisement(int deviceIndex);
+   
+/**
+ Get timecode string representing a cached timecode.
+ 
+ @param deviceIndex Get timecode string for device with this index.
+ @param timecodeIndex Get timecode string for device with this index.
+ @param timecodeCacheString Resulting string value. Has to be of length
+        `TENTACLE_TIMECODE_STRING_LENGTH`, includes terminating null character.
+ @return Number of characters of `timecodeCacheString`, not including the
+         terminating null character.
+ */
+int TentacleDeviceCacheTimecodeString(int deviceIndex, int timecodeIndex, char *timecodeCacheString);
+    
+/**
+ Get timecode string for a specific point in time.
+ 
+ @param deviceIndex Get timecode string for device with this index.
+ @param atTimestamp Get timecode at this timestamp.
+ @param timecodeString Resulting string value. Has to be of length
+        `TENTACLE_TIMECODE_STRING_LENGTH`.
+ @return Number of characters of `timecodeString`, not including the
+         terminating null character.
+ */
+int TentacleDeviceCacheTimecodeStringAtTimestamp(int deviceIndex, double atTimestamp, char *timecodeString);
+
+/**
+ Get timecode string for a specific point in time.
+
+ @param identifier Identifier of device to get.
+ @param identifierLength Length of identifier.
+ @param atTimestamp Get timecode at this timestamp.
+ @param timecodeString Resulting string value. Has to be of length `TENTACLE_TIMECODE_STRING_LENGTH`.
+ @return Number of characters of `timecodeString`, not including the
+ terminating null character.
+ */
+int TentacleDeviceCacheTimecodeStringAtTimestampForIdentifier(char *identifier, size_t identifierLength, double atTimestamp, char *timecodeString);
+
+/**
+ Print string representation of device cache to stdout.
+
+ @param atTimestamp Print timecodes at this timestamp.
+ */
+void TentacleDeviceCachePrint(double atTimestamp);
+    
+#ifdef __cplusplus
+}
+#endif
+
+#endif /* tentacleDeviceCache_h */
diff --git a/Engine/Extras/VirtualProduction/LiveLinkVCAM/Tentacle.xcframework/ios-arm64_armv7/Tentacle.framework/Headers/TentacleEnum.h b/Engine/Extras/VirtualProduction/LiveLinkVCAM/Tentacle.xcframework/ios-arm64_armv7/Tentacle.framework/Headers/TentacleEnum.h
new file mode 100644
--- /dev/null
+++ b/Engine/Extras/VirtualProduction/LiveLinkVCAM/Tentacle.xcframework/ios-arm64_armv7/Tentacle.framework/Headers/TentacleEnum.h
@@ -0,0 +1,92 @@
+//
+//  tentacleEnum.h
+//  Tentacle-SDK
+//
+//  Copyright 2019 Tentacle Sync GmbH
+//
+//  Redistribution and use in source and binary forms, with or without
+//  modification, are permitted provided that the following conditions are met:
+//
+//  1. Redistributions of source code must retain the above copyright notice,
+//     this list of conditions and the following disclaimer.
+//
+//  2. Redistributions in binary form must reproduce the above copyright notice,
+//     this list of conditions and the following disclaimer in the
+//     documentation and/or other materials provided with the distribution.
+//
+//  3. Neither the name of the copyright holder nor the names of its
+//     contributors may be used to endorse or promote products derived from this
+//     software without specific prior written permission.
+//
+//  THIS SOFTWARE IS PROVIDED BY THE COPYRIGHT HOLDERS AND CONTRIBUTORS "AS IS"
+//  AND ANY EXPRESS OR IMPLIED WARRANTIES, INCLUDING, BUT NOT LIMITED TO, THE
+//  IMPLIED WARRANTIES OF MERCHANTABILITY AND FITNESS FOR A PARTICULAR PURPOSE
+//  ARE DISCLAIMED. IN NO EVENT SHALL THE COPYRIGHT HOLDER OR CONTRIBUTORS BE
+//  LIABLE FOR ANY DIRECT, INDIRECT, INCIDENTAL, SPECIAL, EXEMPLARY, OR
+//  CONSEQUENTIAL DAMAGES (INCLUDING, BUT NOT LIMITED TO, PROCUREMENT OF
+//  SUBSTITUTE GOODS OR SERVICES; LOSS OF USE, DATA, OR PROFITS; OR BUSINESS
+//  INTERRUPTION) HOWEVER CAUSED AND ON ANY THEORY OF LIABILITY, WHETHER IN
+//  CONTRACT, STRICT LIABILITY, OR TORT (INCLUDING NEGLIGENCE OR OTHERWISE)
+//  ARISING IN ANY WAY OUT OF THE USE OF THIS SOFTWARE, EVEN IF ADVISED OF THE
+//  POSSIBILITY OF SUCH DAMAGE.
+//
+
+/**
+ @file
+ @brief Tentacle enumerations.
+ */
+
+#ifndef tentacleenum_h
+#define tentacleenum_h
+
+#ifdef TentacleAdvertisement
+extern "C" {
+#endif
+    
+/**
+ State the device is in.
+ */
+enum tentacleDeviceMode {
+    /// Device is powered off.
+    power_off = 0,
+    
+    /// Device is booting.
+    startup = 1,
+    
+    /// Device is in standby.
+    standby = 2,
+    
+    /// No cable is plugged into the device, no audio timecode is generated.
+    unplugged = 3,
+    
+    /// Cable is plugged in and device is generating timecode.
+    generateTimecode = 4,
+    
+    /// Device is reading timecode.
+    readTimecode = 5,
+    
+    /// Device is stopped.
+    stopped = 6
+};
+
+/**
+ When producing timecode to be recorded on an audio track, the timecode is
+ encoded as an audio signal using a certain volume. This volume can be
+ low (Mic), high (Line), something in between (Custom) or be automatically
+ detected (AutoMic).
+ */
+enum TentacleAudioLevel {
+    /// Audio timecode is produced at a low volume
+    mic = 0,
+    
+    /// Audio timecode is produced at a high volume
+    line = 1 ,
+    
+    /// Audio timecode is produced at a medium volume
+    custom = 2,
+    
+    /// Audio timecode is automatically set to a low volume
+    autoMic = 3
+};
+
+#endif /* tentacleenum_h */
diff --git a/Engine/Extras/VirtualProduction/LiveLinkVCAM/Tentacle.xcframework/ios-arm64_armv7/Tentacle.framework/Headers/TentacleIcon.h b/Engine/Extras/VirtualProduction/LiveLinkVCAM/Tentacle.xcframework/ios-arm64_armv7/Tentacle.framework/Headers/TentacleIcon.h
new file mode 100644
--- /dev/null
+++ b/Engine/Extras/VirtualProduction/LiveLinkVCAM/Tentacle.xcframework/ios-arm64_armv7/Tentacle.framework/Headers/TentacleIcon.h
@@ -0,0 +1,80 @@
+//
+//  TentacleIcon.h
+//  Tentacle-SDK
+//
+//  Copyright 2019 Tentacle Sync GmbH
+//
+//  Redistribution and use in source and binary forms, with or without
+//  modification, are permitted provided that the following conditions are met:
+//
+//  1. Redistributions of source code must retain the above copyright notice,
+//     this list of conditions and the following disclaimer.
+//
+//  2. Redistributions in binary form must reproduce the above copyright notice,
+//     this list of conditions and the following disclaimer in the
+//     documentation and/or other materials provided with the distribution.
+//
+//  3. Neither the name of the copyright holder nor the names of its
+//     contributors may be used to endorse or promote products derived from this
+//     software without specific prior written permission.
+//
+//  THIS SOFTWARE IS PROVIDED BY THE COPYRIGHT HOLDERS AND CONTRIBUTORS "AS IS"
+//  AND ANY EXPRESS OR IMPLIED WARRANTIES, INCLUDING, BUT NOT LIMITED TO, THE
+//  IMPLIED WARRANTIES OF MERCHANTABILITY AND FITNESS FOR A PARTICULAR PURPOSE
+//  ARE DISCLAIMED. IN NO EVENT SHALL THE COPYRIGHT HOLDER OR CONTRIBUTORS BE
+//  LIABLE FOR ANY DIRECT, INDIRECT, INCIDENTAL, SPECIAL, EXEMPLARY, OR
+//  CONSEQUENTIAL DAMAGES (INCLUDING, BUT NOT LIMITED TO, PROCUREMENT OF
+//  SUBSTITUTE GOODS OR SERVICES; LOSS OF USE, DATA, OR PROFITS; OR BUSINESS
+//  INTERRUPTION) HOWEVER CAUSED AND ON ANY THEORY OF LIABILITY, WHETHER IN
+//  CONTRACT, STRICT LIABILITY, OR TORT (INCLUDING NEGLIGENCE OR OTHERWISE)
+//  ARISING IN ANY WAY OUT OF THE USE OF THIS SOFTWARE, EVEN IF ADVISED OF THE
+//  POSSIBILITY OF SUCH DAMAGE.
+//
+
+@import Foundation;
+
+#if TARGET_OS_IOS
+@import UIKit;
+#endif
+
+NS_ASSUME_NONNULL_BEGIN
+
+@interface TentacleIcon: NSObject
+
+#if TARGET_OS_IOS
+
+/// Get Tentacle device icon image.
+/// \param index Icon index of range 0 to 37
+///
+/// returns:
+/// Device image
++ (UIImage * _Nonnull)tentacleIconWithIndex:(NSUInteger)index;
+
+/// Get default battery image.
+/// \param percent Percentage of battery status
+///
+/// returns:
+/// Battery image
++ (UIImage * _Nonnull)tentacleBatteryIconForPercent:(NSUInteger)percent;
+
+/// Get default signal strength image.
+/// \param index Possible values are 0 to 3
+///
+/// returns:
+/// Signal strength image
++ (UIImage * _Nonnull)tentacleSignalStrengthIconWithIndex:(NSUInteger)index;
+
+/// Get default signal strength image for a given received signal strength indicator (RSSI) value.
+/// \param rssiValue Get image corresponding to this RSSI
+/// \param linkMode Internal use only, always pass false
+///
+/// returns:
+/// Signal strength image
++ (UIImage * _Nonnull)tentacleSignalStrengthIconForRssiValue:(NSInteger)rssiValue linkMode:(BOOL)linkMode;
+
+#endif // #if TARGET_OS_IOS
+
+@end
+
+NS_ASSUME_NONNULL_END
+
diff --git a/Engine/Extras/VirtualProduction/LiveLinkVCAM/Tentacle.xcframework/ios-arm64_armv7/Tentacle.framework/Headers/TentacleLog.h b/Engine/Extras/VirtualProduction/LiveLinkVCAM/Tentacle.xcframework/ios-arm64_armv7/Tentacle.framework/Headers/TentacleLog.h
new file mode 100644
--- /dev/null
+++ b/Engine/Extras/VirtualProduction/LiveLinkVCAM/Tentacle.xcframework/ios-arm64_armv7/Tentacle.framework/Headers/TentacleLog.h
@@ -0,0 +1,66 @@
+//
+//  tentacleLog.h
+//  Tentacle-SDK
+//
+//  Copyright 2019 Tentacle Sync GmbH
+//
+//  Redistribution and use in source and binary forms, with or without
+//  modification, are permitted provided that the following conditions are met:
+//
+//  1. Redistributions of source code must retain the above copyright notice,
+//     this list of conditions and the following disclaimer.
+//
+//  2. Redistributions in binary form must reproduce the above copyright notice,
+//     this list of conditions and the following disclaimer in the
+//     documentation and/or other materials provided with the distribution.
+//
+//  3. Neither the name of the copyright holder nor the names of its
+//     contributors may be used to endorse or promote products derived from this
+//     software without specific prior written permission.
+//
+//  THIS SOFTWARE IS PROVIDED BY THE COPYRIGHT HOLDERS AND CONTRIBUTORS "AS IS"
+//  AND ANY EXPRESS OR IMPLIED WARRANTIES, INCLUDING, BUT NOT LIMITED TO, THE
+//  IMPLIED WARRANTIES OF MERCHANTABILITY AND FITNESS FOR A PARTICULAR PURPOSE
+//  ARE DISCLAIMED. IN NO EVENT SHALL THE COPYRIGHT HOLDER OR CONTRIBUTORS BE
+//  LIABLE FOR ANY DIRECT, INDIRECT, INCIDENTAL, SPECIAL, EXEMPLARY, OR
+//  CONSEQUENTIAL DAMAGES (INCLUDING, BUT NOT LIMITED TO, PROCUREMENT OF
+//  SUBSTITUTE GOODS OR SERVICES; LOSS OF USE, DATA, OR PROFITS; OR BUSINESS
+//  INTERRUPTION) HOWEVER CAUSED AND ON ANY THEORY OF LIABILITY, WHETHER IN
+//  CONTRACT, STRICT LIABILITY, OR TORT (INCLUDING NEGLIGENCE OR OTHERWISE)
+//  ARISING IN ANY WAY OUT OF THE USE OF THIS SOFTWARE, EVEN IF ADVISED OF THE
+//  POSSIBILITY OF SUCH DAMAGE.
+//
+
+/**
+ @file
+ @brief Discovered devices.
+ */
+
+#ifndef tentaclelog_h
+#define tentaclelog_h
+
+#ifdef __cplusplus
+extern "C"
+{
+#endif
+
+#include <stdio.h>
+#include <stdarg.h>
+
+#if __ANDROID__ == 1
+#include <android/log.h>
+#endif
+
+/**
+ * Prints logging information using platform specific methods.
+ *
+ * @param fmt Format string.
+ * @param ... Variadic arguments to be logged.
+ */
+void TentacleLog(const char *fmt, ...);
+
+#ifdef __cplusplus
+}
+#endif
+
+#endif /* tentaclelog_h */
diff --git a/Engine/Extras/VirtualProduction/LiveLinkVCAM/Tentacle.xcframework/ios-arm64_armv7/Tentacle.framework/Headers/TentacleTimecode.h b/Engine/Extras/VirtualProduction/LiveLinkVCAM/Tentacle.xcframework/ios-arm64_armv7/Tentacle.framework/Headers/TentacleTimecode.h
new file mode 100644
--- /dev/null
+++ b/Engine/Extras/VirtualProduction/LiveLinkVCAM/Tentacle.xcframework/ios-arm64_armv7/Tentacle.framework/Headers/TentacleTimecode.h
@@ -0,0 +1,172 @@
+//
+//  tentacleTimecode.h
+//  Tentacle-SDK
+//
+//  Copyright 2019 Tentacle Sync GmbH
+//
+//  Redistribution and use in source and binary forms, with or without
+//  modification, are permitted provided that the following conditions are met:
+//
+//  1. Redistributions of source code must retain the above copyright notice,
+//     this list of conditions and the following disclaimer.
+//
+//  2. Redistributions in binary form must reproduce the above copyright notice,
+//     this list of conditions and the following disclaimer in the
+//     documentation and/or other materials provided with the distribution.
+//
+//  3. Neither the name of the copyright holder nor the names of its
+//     contributors may be used to endorse or promote products derived from this
+//     software without specific prior written permission.
+//
+//  THIS SOFTWARE IS PROVIDED BY THE COPYRIGHT HOLDERS AND CONTRIBUTORS "AS IS"
+//  AND ANY EXPRESS OR IMPLIED WARRANTIES, INCLUDING, BUT NOT LIMITED TO, THE
+//  IMPLIED WARRANTIES OF MERCHANTABILITY AND FITNESS FOR A PARTICULAR PURPOSE
+//  ARE DISCLAIMED. IN NO EVENT SHALL THE COPYRIGHT HOLDER OR CONTRIBUTORS BE
+//  LIABLE FOR ANY DIRECT, INDIRECT, INCIDENTAL, SPECIAL, EXEMPLARY, OR
+//  CONSEQUENTIAL DAMAGES (INCLUDING, BUT NOT LIMITED TO, PROCUREMENT OF
+//  SUBSTITUTE GOODS OR SERVICES; LOSS OF USE, DATA, OR PROFITS; OR BUSINESS
+//  INTERRUPTION) HOWEVER CAUSED AND ON ANY THEORY OF LIABILITY, WHETHER IN
+//  CONTRACT, STRICT LIABILITY, OR TORT (INCLUDING NEGLIGENCE OR OTHERWISE)
+//  ARISING IN ANY WAY OUT OF THE USE OF THIS SOFTWARE, EVEN IF ADVISED OF THE
+//  POSSIBILITY OF SUCH DAMAGE.
+//
+
+/**
+ @file
+ @brief Tentacle timecode.
+ 
+ Contains cached raw data of a received timecode as received from a Bluetooth
+ advertisement.
+ */
+
+#ifndef tentacletimecode_h
+#define tentacletimecode_h
+
+#ifdef __cplusplus
+extern "C" {
+#endif
+
+#include <stdio.h>   // for size_t
+#include <stdint.h>  // for uint8_t
+#include <stdbool.h> // for bool
+    
+/// Length of string returned by `tentacleTimecodeCacheGetString`.
+#define TENTACLE_TIMECODE_STRING_LENGTH 35
+    
+/// 23.98 FPS =  24 FPS NTSC
+#define TENTACLE_FPS_23_98 (24000.0/1001.0)
+    
+/// 29.97 FPS = 30 FPS NTSC
+#define TENTACLE_FPS_29_97 (30000.0/1001.0)
+
+/**
+ Contains cached raw data of a received timecode as received from a Bluetooth
+ advertisement.
+ */
+typedef struct TentacleTimecode {
+    /// Hour of timecode.
+    uint8_t hour;
+    
+    /// Minute of timecode.
+    uint8_t minute;
+    
+    /// Second of timecode.
+    uint8_t second;
+    
+    /// Frame of timecode.
+    uint8_t frame;
+    
+    /// Microseconds since start of last frame.
+    uint16_t microsecond;
+    
+    /// \brief Timestamp of when this timecode has been received, in seconds.
+    ///
+    /// On Apple Darwin, [machAbsoluteTime]([REDACTED_PROFILE_URL]/darwin-xnu/blob/0a798f6738bc1db01281fc08ae024145e84df927/osfmk/arm/rtclock.c#L162) is used.
+    ///
+    /// On Android, [ScanResult.getTimestampNanos()](https://developer.android.com/reference/android/bluetooth/le/ScanResult.html#getTimestampNanos()) is used.
+    ///
+    /// Caution: Due to maximum precision these Timestamps are based on the boot-time of the device. According to this, the timestamps are invalid after
+    /// a device reboot and should not be used for long time comparison or alike.
+    double receivedTimestamp;
+} TentacleTimecode;
+    
+/**
+ Returns a timecode struct with default values.
+
+ @return TentacleTimecode struct with default values.
+ */
+TentacleTimecode TentacleTimecodeMake(void);
+    
+/**
+ Check if two timecodes are equal.
+
+ @param timecode1 First timecode to compare.
+ @param timecode2 Second timecode to compare.
+ @return True if timecodes are equal, false otherwise.
+ */
+bool TentacleTimecodeEqual(const TentacleTimecode *timecode1, const TentacleTimecode *timecode2);
+
+/**
+ Total number of seconds described by this timecode at specific timestamp.
+
+ @param timecode Get total number of seconds for this timecode.
+ @param frameRate Frame rate used to calculate number of seconds.
+ @param dropFrame True if timecode is using drop frames.
+ @param atTimestamp Total number of seconds of timecode at this timestamp.
+ @return Frame rate used to calculate duration of frame count at timestamp.
+ */
+double TentacleTimecodeSecondsAtTimestamp(const TentacleTimecode *timecode, double frameRate, bool dropFrame, double atTimestamp);
+    
+/**
+ Create a timecode from a total number of seconds.
+
+ @param seconds Convert this total number of seconds to timecode.
+ @param frameRate Frame rate used during conversion.
+ @param dropFrame True if timecode is using drop frames.
+ @param receivedTimestamp Use this as received timestamp.
+ @return New timecode struct containing the passed total number of seconds.
+ */
+TentacleTimecode TentacleTimecodeFromSeconds(double seconds, double frameRate, bool dropFrame, double receivedTimestamp);
+
+/**
+ Given a received timecode, calculate what its value would be at a different
+ point in time.
+
+ @param timecode Get this timecode for a specific point in time
+ @param frameRate Frame rate of timecode
+ @param dropFrame True if timecode is using drop frames.
+ @param atTimestamp Requested point in time
+ @return Timecode value at specific point in time.
+ */
+TentacleTimecode TentacleTimecodeAtTimestamp(TentacleTimecode timecode, double frameRate, bool dropFrame, double atTimestamp);
+
+/**
+ Get timecode string representing a timecode.
+ 
+ @param timecode Get string value representing this timecode struct.
+ @param dropFrame True if timecode is drop frame.
+ @param timecodeCacheString Resulting string value. Has to be of length
+        `TENTACLE_TIMECODE_STRING_LENGTH`.
+ @return Number of characters of `timecodeCacheString`, not including the
+         terminating null character.
+ */
+int TentacleTimecodeString(TentacleTimecode timecode, bool dropFrame, char *timecodeCacheString);
+
+/**
+ Get timecode string representing a timecode, including microseconds since last
+ frame and received timestamp.
+ 
+ @param timecode Get string value representing this timecode struct.
+ @param dropFrame True if timecode is drop frame.
+ @param timecodeCacheString Resulting string value. Has to be of length
+        `TENTACLE_TIMECODE_STRING_LENGTH`.
+ @return Number of characters of `timecodeCacheString`, not including the
+         terminating null character.
+ */
+int TentacleTimecodeStringVerbose(TentacleTimecode timecode, bool dropFrame, char *timecodeCacheString);
+
+#ifdef __cplusplus
+}
+#endif
+
+#endif /* tentacletimecode_h */
diff --git a/Engine/Extras/VirtualProduction/LiveLinkVCAM/Tentacle.xcframework/ios-arm64_i386_x86_64-simulator/Tentacle.framework/Headers/Tentacle.h b/Engine/Extras/VirtualProduction/LiveLinkVCAM/Tentacle.xcframework/ios-arm64_i386_x86_64-simulator/Tentacle.framework/Headers/Tentacle.h
new file mode 100644
--- /dev/null
+++ b/Engine/Extras/VirtualProduction/LiveLinkVCAM/Tentacle.xcframework/ios-arm64_i386_x86_64-simulator/Tentacle.framework/Headers/Tentacle.h
@@ -0,0 +1,52 @@
+//
+//  Tentacle.h
+//  TentacleFramework
+//
+//  Copyright 2019 Tentacle Sync GmbH
+//
+//  Redistribution and use in source and binary forms, with or without
+//  modification, are permitted provided that the following conditions are met:
+//
+//  1. Redistributions of source code must retain the above copyright notice,
+//     this list of conditions and the following disclaimer.
+//
+//  2. Redistributions in binary form must reproduce the above copyright notice,
+//     this list of conditions and the following disclaimer in the
+//     documentation and/or other materials provided with the distribution.
+//
+//  3. Neither the name of the copyright holder nor the names of its
+//     contributors may be used to endorse or promote products derived from this
+//     software without specific prior written permission.
+//
+//  THIS SOFTWARE IS PROVIDED BY THE COPYRIGHT HOLDERS AND CONTRIBUTORS "AS IS"
+//  AND ANY EXPRESS OR IMPLIED WARRANTIES, INCLUDING, BUT NOT LIMITED TO, THE
+//  IMPLIED WARRANTIES OF MERCHANTABILITY AND FITNESS FOR A PARTICULAR PURPOSE
+//  ARE DISCLAIMED. IN NO EVENT SHALL THE COPYRIGHT HOLDER OR CONTRIBUTORS BE
+//  LIABLE FOR ANY DIRECT, INDIRECT, INCIDENTAL, SPECIAL, EXEMPLARY, OR
+//  CONSEQUENTIAL DAMAGES (INCLUDING, BUT NOT LIMITED TO, PROCUREMENT OF
+//  SUBSTITUTE GOODS OR SERVICES; LOSS OF USE, DATA, OR PROFITS; OR BUSINESS
+//  INTERRUPTION) HOWEVER CAUSED AND ON ANY THEORY OF LIABILITY, WHETHER IN
+//  CONTRACT, STRICT LIABILITY, OR TORT (INCLUDING NEGLIGENCE OR OTHERWISE)
+//  ARISING IN ANY WAY OUT OF THE USE OF THIS SOFTWARE, EVEN IF ADVISED OF THE
+//  POSSIBILITY OF SUCH DAMAGE.
+//
+
+#import <Foundation/Foundation.h>
+
+//! Project version number for TentacleFramework.
+FOUNDATION_EXPORT double TentacleFrameworkVersionNumber;
+
+//! Project version string for TentacleFramework.
+FOUNDATION_EXPORT const unsigned char TentacleFrameworkVersionString[];
+
+// In this header, you should import all the public headers of your framework using statements like #import <Tentacle/PublicHeader.h>
+#import <Tentacle/TentacleTimecode.h>
+#import <Tentacle/TentacleAdvertisement.h>
+#import <Tentacle/TentacleDevice.h>
+#import <Tentacle/TentacleDeviceCache.h>
+#import <Tentacle/TentacleBluetoothController.h>
+#import <Tentacle/TentacleLog.h>
+
+#if TARGET_OS_IOS
+#import <Tentacle/TentacleIcon.h>
+#endif
diff --git a/Engine/Extras/VirtualProduction/LiveLinkVCAM/Tentacle.xcframework/ios-arm64_i386_x86_64-simulator/Tentacle.framework/Headers/TentacleAdvertisement.h b/Engine/Extras/VirtualProduction/LiveLinkVCAM/Tentacle.xcframework/ios-arm64_i386_x86_64-simulator/Tentacle.framework/Headers/TentacleAdvertisement.h
new file mode 100644
--- /dev/null
+++ b/Engine/Extras/VirtualProduction/LiveLinkVCAM/Tentacle.xcframework/ios-arm64_i386_x86_64-simulator/Tentacle.framework/Headers/TentacleAdvertisement.h
@@ -0,0 +1,223 @@
+//
+//  tentacleAdvertisement.h
+//  Tentacle-SDK
+//
+//  Copyright 2019 Tentacle Sync GmbH
+//
+//  Redistribution and use in source and binary forms, with or without
+//  modification, are permitted provided that the following conditions are met:
+//
+//  1. Redistributions of source code must retain the above copyright notice,
+//     this list of conditions and the following disclaimer.
+//
+//  2. Redistributions in binary form must reproduce the above copyright notice,
+//     this list of conditions and the following disclaimer in the
+//     documentation and/or other materials provided with the distribution.
+//
+//  3. Neither the name of the copyright holder nor the names of its
+//     contributors may be used to endorse or promote products derived from this
+//     software without specific prior written permission.
+//
+//  THIS SOFTWARE IS PROVIDED BY THE COPYRIGHT HOLDERS AND CONTRIBUTORS "AS IS"
+//  AND ANY EXPRESS OR IMPLIED WARRANTIES, INCLUDING, BUT NOT LIMITED TO, THE
+//  IMPLIED WARRANTIES OF MERCHANTABILITY AND FITNESS FOR A PARTICULAR PURPOSE
+//  ARE DISCLAIMED. IN NO EVENT SHALL THE COPYRIGHT HOLDER OR CONTRIBUTORS BE
+//  LIABLE FOR ANY DIRECT, INDIRECT, INCIDENTAL, SPECIAL, EXEMPLARY, OR
+//  CONSEQUENTIAL DAMAGES (INCLUDING, BUT NOT LIMITED TO, PROCUREMENT OF
+//  SUBSTITUTE GOODS OR SERVICES; LOSS OF USE, DATA, OR PROFITS; OR BUSINESS
+//  INTERRUPTION) HOWEVER CAUSED AND ON ANY THEORY OF LIABILITY, WHETHER IN
+//  CONTRACT, STRICT LIABILITY, OR TORT (INCLUDING NEGLIGENCE OR OTHERWISE)
+//  ARISING IN ANY WAY OUT OF THE USE OF THIS SOFTWARE, EVEN IF ADVISED OF THE
+//  POSSIBILITY OF SUCH DAMAGE.
+//
+
+/**
+ @file
+ @brief Advertisement data.
+ 
+ Contains data of a received Bluetooth advertisement packet.
+ */
+
+#ifndef tentacleAdvertisement_h
+#define tentacleAdvertisement_h
+
+#ifdef __cplusplus
+extern "C" {
+#endif
+
+#include <stdio.h>
+#include <stdbool.h>
+#include <string.h>
+#include <Tentacle/TentacleEnum.h>
+#include <Tentacle/TentacleTimecode.h>
+
+/// Length of string returned by `tentacleAdvertisementString`.
+#define TENTACLE_ADVERTISEMENT_STRING_LENGTH 104
+    
+/// Length of unique identifier representing a discovered device, see `tentacleAdvertisement.identifier`.
+#define TENTACLE_ADVERTISEMENT_ID_LENGTH_MAX 37
+    
+/// Maximum length of a device name, see `tentacleAdvertisement.name`.
+#define TENTACLE_ADVERTISEMENT_NAME_LENGTH_MAX 16
+
+/**
+ Contains data of a received Bluetooth advertisement packet.
+ */
+typedef struct TentacleAdvertisement {
+    /// Tentacle Bluetooth protocol version. Currently supported version is 1.
+    uint8_t version;
+    
+    /// True if device is able to connect via Bluetooth.
+    bool connectable;
+    
+    /// True if a mobile device can authenticate to the device.
+    bool linkMode;
+    
+    /// Device type. Currently, the only valid value is 0 for Tentacle Sync-E.
+    uint8_t deviceType;
+    
+    /// State the device is in.
+    enum tentacleDeviceMode deviceMode;
+    
+    /// Audio volume when generating audio timecode.
+    enum TentacleAudioLevel audioLevel;
+    
+    /// True if device is generating NTSC timecode.
+    bool ntsc;
+    
+    /// True if device is generating drop frame timecode.
+    bool dropFrame;
+    
+    /// \brief True if device is in green mode.
+    ///
+    /// Green mode means that the timecode has been configured, the device is
+    /// generating audio timecode and other devices can be synced to this device
+    /// using a cable.
+    bool greenMode;
+    
+    /// Hour, minute, second, frame, microseconds and reception timestamp
+    TentacleTimecode timecode;
+    
+    /// Battery level of device in percent.
+    uint8_t battery;
+    
+    /// True if device is currently charging.
+    bool charging;
+    
+    /// Icon index the device is set to.
+    uint8_t icon;
+    
+    /// Received Signal Strength Indication
+    int8_t rssi;
+    
+    /// \brief Unique identifier string representing the device this
+    /// advertisement was sent from.
+    ///
+    /// On iOS this is the string value of a [CBUUID](https://developer.apple.com/documentation/corebluetooth/cbuuid) object.
+    ///
+    /// On Android this is the hardware address as returned by [BluetoothDevice.getAddress()](https://developer.android.com/reference/android/bluetooth/BluetoothDevice.html#getAddress()),
+    ///
+    /// On Nordic nRF5 SDK this is the hardware address as passed in  [bleGapAddr_t](https://www.nordicsemi.com/DocLib/Content/SoftDeviceAPIDoc/S132/v6-1-0/structble_Gap_Addr__t).
+    char identifier[TENTACLE_ADVERTISEMENT_ID_LENGTH_MAX];
+    
+    /// Length of identifier string, has to be less than `tentacleAdvertisementIdentifierLength`.
+    size_t identifierLength;
+    
+    /// Name of the device this advertisement was sent from.
+    char name[TENTACLE_ADVERTISEMENT_NAME_LENGTH_MAX];
+    
+    /// Length of name string, has to be less than `tentacleNameLength`.
+    size_t nameLength;
+    
+    /// @brief Current frame rate as integer value, ignoring NTSC flag.
+    ///
+    /// To calculate the actual frame rate, you have to take the NTSC flag into
+    /// account or use `tentacleAdvertisementGetFrameRate()`.
+    int frameRateInt;
+    
+    /// @brief True if data in this struct is considered to describe a valid timecode.
+    ///
+    /// Value is set by `tentacleAdvertisementInit()`.
+    bool valid;
+} TentacleAdvertisement;
+    
+/**
+ Returns an advertisement struct with default values.
+ 
+ @return TentacleAdvertisement struct with default values.
+ */
+TentacleAdvertisement TentacleAdvertisementMake(void);
+
+/**
+ Initializes a `tentacleAdvertisement` from received Bluetooth advertisement
+ packet data.
+
+ If the passed data is considered to be valid, `advertisement.valid` is set to true.
+
+ @param advertisementData Byte array containing advertisement packet data.
+ @param advertisementDataLen Length of advertisement data array.
+ @param rssi Received Signal Strength Indication.
+ @param receivedTimestamp Precise timestamp of when this packet was received,
+        in seconds.
+ @param identifierBuffer Unique identifier string representing the device this
+        advertisement was sent from.
+ @param identifierBufferLength Length of identifier string, has to be less than
+        `tentacleAdvertisementIdentifierLengthMax`.
+ @param nameBuffer Name of the device this advertisement was sent from.
+ @param nameBufferSize Length of name string, has to be less than
+        `tentacleNameLength`.
+ @return Advertisement struct containing received packet data.
+ */
+TentacleAdvertisement TentacleAdvertisementInit(const unsigned char *advertisementData, size_t advertisementDataLen, int8_t rssi, double receivedTimestamp, const char *identifierBuffer, size_t identifierBufferLength, const char *nameBuffer, size_t nameBufferSize);
+
+/**
+ Check if two frame rates are equal.
+ 
+ @param advertisement1 First frame rate to compare.
+ @param advertisement2 Second frame rate to compare.
+ @return True if frame rate are equal, false otherwise.
+ */
+bool TentacleAdvertisementFrameRateEqual(const TentacleAdvertisement *advertisement1, const TentacleAdvertisement *advertisement2);
+
+/**
+ Frame rate, depending on value of NTSC flag. Also see
+ `tentacleAdvertisement.frameRateInt` for frame rate ignoring the NTSC flag.
+
+ @param advertisement advertisement to get the frame rate of.
+ @return frame rate depending on NTSC flag.
+ */
+double TentacleAdvertisementGetFrameRate(const TentacleAdvertisement *advertisement);
+
+/**
+ String value of advertisement, containing frame rate, battery state and more.
+
+ @param advertisement Get string value of this advertisement.
+ @param advertisementString Resulting string value. Has to be of length
+        `TENTACLE_ADVERTISEMENT_STRING_LENGTH`, includes terminating null
+        character.
+ @return Number of characters of `advertisementString`, not including the
+         terminating null character.
+ */
+int TentacleAdvertisementString(const TentacleAdvertisement *advertisement, char *advertisementString);
+    
+/**
+ Name of the device this advertisement was received from.
+
+ @param advertisement advertisement to get the name of.
+ @param nameString Resulting name string
+ @return Length of returned string.
+ */
+int TentacleAdvertisementNameString(const TentacleAdvertisement *advertisement, char *nameString);
+
+/**
+ Print string representation of advertisement to stdout.
+
+ @param advertisement Advertisement struct to print.
+ */
+void TentacleAdvertisementPrint(const TentacleAdvertisement *advertisement);
+
+#ifdef __cplusplus
+}
+#endif
+
+#endif /* tentacleAdvertisement_h */
diff --git a/Engine/Extras/VirtualProduction/LiveLinkVCAM/Tentacle.xcframework/ios-arm64_i386_x86_64-simulator/Tentacle.framework/Headers/TentacleBluetoothController.h b/Engine/Extras/VirtualProduction/LiveLinkVCAM/Tentacle.xcframework/ios-arm64_i386_x86_64-simulator/Tentacle.framework/Headers/TentacleBluetoothController.h
new file mode 100644
--- /dev/null
+++ b/Engine/Extras/VirtualProduction/LiveLinkVCAM/Tentacle.xcframework/ios-arm64_i386_x86_64-simulator/Tentacle.framework/Headers/TentacleBluetoothController.h
@@ -0,0 +1,117 @@
+//
+//  TentacleBluetoothController.h
+//  Tentacle-SDK
+//
+//  Copyright 2019 Tentacle Sync GmbH
+//
+//  Redistribution and use in source and binary forms, with or without
+//  modification, are permitted provided that the following conditions are met:
+//
+//  1. Redistributions of source code must retain the above copyright notice,
+//     this list of conditions and the following disclaimer.
+//
+//  2. Redistributions in binary form must reproduce the above copyright notice,
+//     this list of conditions and the following disclaimer in the
+//     documentation and/or other materials provided with the distribution.
+//
+//  3. Neither the name of the copyright holder nor the names of its
+//     contributors may be used to endorse or promote products derived from this
+//     software without specific prior written permission.
+//
+//  THIS SOFTWARE IS PROVIDED BY THE COPYRIGHT HOLDERS AND CONTRIBUTORS "AS IS"
+//  AND ANY EXPRESS OR IMPLIED WARRANTIES, INCLUDING, BUT NOT LIMITED TO, THE
+//  IMPLIED WARRANTIES OF MERCHANTABILITY AND FITNESS FOR A PARTICULAR PURPOSE
+//  ARE DISCLAIMED. IN NO EVENT SHALL THE COPYRIGHT HOLDER OR CONTRIBUTORS BE
+//  LIABLE FOR ANY DIRECT, INDIRECT, INCIDENTAL, SPECIAL, EXEMPLARY, OR
+//  CONSEQUENTIAL DAMAGES (INCLUDING, BUT NOT LIMITED TO, PROCUREMENT OF
+//  SUBSTITUTE GOODS OR SERVICES; LOSS OF USE, DATA, OR PROFITS; OR BUSINESS
+//  INTERRUPTION) HOWEVER CAUSED AND ON ANY THEORY OF LIABILITY, WHETHER IN
+//  CONTRACT, STRICT LIABILITY, OR TORT (INCLUDING NEGLIGENCE OR OTHERWISE)
+//  ARISING IN ANY WAY OUT OF THE USE OF THIS SOFTWARE, EVEN IF ADVISED OF THE
+//  POSSIBILITY OF SUCH DAMAGE.
+//
+
+@import CoreBluetooth;
+#include <Tentacle/TentacleDevice.h>
+
+NS_ASSUME_NONNULL_BEGIN
+
+/// State of the Bluetooth controller. Uses the same values as <code>CBManagerState</code>.
+/// See https://developer.apple.com/documentation/corebluetooth/cbmanagerstate
+typedef NS_ENUM(NSUInteger, TentacleBluetoothState) {
+    TentacleBluetoothStateUnknown = 0,
+    TentacleBluetoothStateResetting = 1,
+    TentacleBluetoothStateUnsupported = 2,
+    TentacleBluetoothStateUnauthorized = 3,
+    TentacleBluetoothStatePoweredOff = 4,
+    TentacleBluetoothStatePoweredOn = 5,
+};
+
+/// The delegate of a <code>TentacleBluetoothController</code> object must adopt
+/// this protocol. The delegate uses this protocol’s methods to notify about the
+/// discovery, state changes and disappearance of Bluetooth devices.
+@protocol TentacleBluetoothControllerDelegate <NSObject>
+@optional
+/// Called if state of <code>TentacleBluetoothController</code> changed.
+/// \param state New state the Bluetooth controller is in.
+- (void)didUpdateToState:(enum TentacleBluetoothState)state;
+
+/// Called if scanning for Bluetooth devices has started.
+- (void)didStartScanning;
+
+@required
+/// Called each time advertisment data has been received from a Tentacle.
+/// Advertisment broadcast rate is multiple times per second, thus this may
+/// be called frequently.
+/// \param deviceIndex Device with this index has been updated
+- (void)didReceiveAdvertismentForDeviceIndex:(NSInteger)deviceIndex;
+@end
+
+
+@interface TentacleBluetoothController : NSObject <CBCentralManagerDelegate>
+/// Dictionary of all devices which have been seen.
+/// Peripherals are only added to this array but never removed.
+@property (nonatomic, strong, readonly) NSMutableDictionary<NSString *, CBPeripheral *> * _Nonnull discoveredPeripherals;
+    
+/// <code>CBCentralManager</code> instance used to scan for devices.
+@property (nonatomic, strong, readonly) CBCentralManager * _Nonnull centralManager;
+
+/// UUID of the Tentacle Bluetooth service.
+@property (nonatomic, strong, readonly) CBUUID * _Nonnull tentacleServiceUUID;
+    
+/// Delegate to receive the Bluetooth events.
+@property (nonatomic, weak) id <TentacleBluetoothControllerDelegate> delegate;
+
+/// Resets the central manager
+- (void)reset;
+
+/// Starts scanning for Bluetooth Tentacles. Scanning is started with the
+/// <code>CBCentralManagerScanOptionAllowDuplicatesKey</code> option set.
+- (void)startScanning;
+
+/// Stops scanning for Bluetooth Tentacles.
+- (void)stopScanning;
+
+/// Returns a Tentacle device with a specific index from device cache.
+/// \param deviceIndex Get device with this index in device cache.
+///
+/// returns:
+/// Device with specified index.
+- (TentacleDevice)retrieveDeviceWithIndex:(int)deviceIndex;
+
+/// Precise timestamp in seconds.
+///
+/// Uses the `mach_absolute_time()` mach kernel syscall.
+///
+/// See section
+/// "Converting Mach absolute time to nanoseconds using mach_timebase_info" in
+/// Apple's technical Q&A [Mach Absolute Time Units](https://developer.apple.com/library/content/qa/qa1398/_index.html)
+/// and the [mach_time.h]([REDACTED_PROFILE_URL]/darwin-xnu/blob/master/osfmk/mach/mach_time.h)
+/// kernel sources for the [mach_timebase_info]([REDACTED_PROFILE_URL]/darwin-xnu/blob/xnu-4570.1.46/osfmk/mach/mach_time.h#L36) struct.
+///
+/// returns:
+/// Precise timestamp in seconds.
++ (double)timestamp;
+@end
+
+NS_ASSUME_NONNULL_END
diff --git a/Engine/Extras/VirtualProduction/LiveLinkVCAM/Tentacle.xcframework/ios-arm64_i386_x86_64-simulator/Tentacle.framework/Headers/TentacleDevice.h b/Engine/Extras/VirtualProduction/LiveLinkVCAM/Tentacle.xcframework/ios-arm64_i386_x86_64-simulator/Tentacle.framework/Headers/TentacleDevice.h
new file mode 100644
--- /dev/null
+++ b/Engine/Extras/VirtualProduction/LiveLinkVCAM/Tentacle.xcframework/ios-arm64_i386_x86_64-simulator/Tentacle.framework/Headers/TentacleDevice.h
@@ -0,0 +1,163 @@
+//
+//  tentacleDevice.h
+//  Tentacle-SDK
+//
+//  Copyright 2019 Tentacle Sync GmbH
+//
+//  Redistribution and use in source and binary forms, with or without
+//  modification, are permitted provided that the following conditions are met:
+//
+//  1. Redistributions of source code must retain the above copyright notice,
+//     this list of conditions and the following disclaimer.
+//
+//  2. Redistributions in binary form must reproduce the above copyright notice,
+//     this list of conditions and the following disclaimer in the
+//     documentation and/or other materials provided with the distribution.
+//
+//  3. Neither the name of the copyright holder nor the names of its
+//     contributors may be used to endorse or promote products derived from this
+//     software without specific prior written permission.
+//
+//  THIS SOFTWARE IS PROVIDED BY THE COPYRIGHT HOLDERS AND CONTRIBUTORS "AS IS"
+//  AND ANY EXPRESS OR IMPLIED WARRANTIES, INCLUDING, BUT NOT LIMITED TO, THE
+//  IMPLIED WARRANTIES OF MERCHANTABILITY AND FITNESS FOR A PARTICULAR PURPOSE
+//  ARE DISCLAIMED. IN NO EVENT SHALL THE COPYRIGHT HOLDER OR CONTRIBUTORS BE
+//  LIABLE FOR ANY DIRECT, INDIRECT, INCIDENTAL, SPECIAL, EXEMPLARY, OR
+//  CONSEQUENTIAL DAMAGES (INCLUDING, BUT NOT LIMITED TO, PROCUREMENT OF
+//  SUBSTITUTE GOODS OR SERVICES; LOSS OF USE, DATA, OR PROFITS; OR BUSINESS
+//  INTERRUPTION) HOWEVER CAUSED AND ON ANY THEORY OF LIABILITY, WHETHER IN
+//  CONTRACT, STRICT LIABILITY, OR TORT (INCLUDING NEGLIGENCE OR OTHERWISE)
+//  ARISING IN ANY WAY OUT OF THE USE OF THIS SOFTWARE, EVEN IF ADVISED OF THE
+//  POSSIBILITY OF SUCH DAMAGE.
+//
+
+/**
+ @file
+ @brief Tentacle device.
+ 
+ Cached data of a discovered device.
+ */
+
+#ifndef tentacledevice_h
+#define tentacledevice_h
+
+#ifdef __cplusplus
+extern "C" {
+#endif
+    
+#include <stdio.h>
+#include <Tentacle/TentacleTimecode.h>
+#include <Tentacle/TentacleAdvertisement.h>
+    
+/// Maximum of timecode values to cache for each device
+#define TENTACLE_DEVICE_TIMECODE_CACHE_SIZE_MAX 10
+
+/// Maximum of received signal strength indicator (RSSI) values for each device
+#define TENTACLE_DEVICE_RSSI_CACHE_SIZE_MAX 5
+
+/// If a device has not been seen after the defined time span it is considered
+/// disappeared.
+#define TENTACLE_DEVICE_DISAPPEARED_THRESHOLD 10.0
+
+/// If a device has not been seen after the defined time span it is considered
+/// unavailable.
+#define TENTACLE_DEVICE_UNAVAILABLE_THRESHOLD 600.0
+
+/**
+ Cached data of a discovered device.
+ */
+typedef struct TentacleDevice {
+    /// Last received advertisement for this device
+    struct TentacleAdvertisement advertisement;
+    
+    /// Ring buffer of cached timecodes
+    TentacleTimecode timecodeCache[TENTACLE_DEVICE_TIMECODE_CACHE_SIZE_MAX];
+
+    /// Current index of timecode in ring buffer
+    size_t timecodeCacheIndex;
+    
+    /// Current size of timecode ring buffer
+    size_t timecodeCacheSize;
+    
+    /// Averaged timecode of all timecodes from timecode cache.
+    TentacleTimecode timecode;
+    
+    /// Ring buffer of received signal strength indication (RSSI) values
+    int rssiCache[TENTACLE_DEVICE_RSSI_CACHE_SIZE_MAX];
+    
+    /// Next index in RSSI ring buffer to use
+    size_t rssiCacheNextIndex;
+    
+    /// @private Used to calculate precise timecode
+    int startPhaseCounter;
+    
+    /// @private Used to calculate precise timecode
+    int asyncCounter;
+    
+    /// @private Used to calculate precise timecode
+    int severeAsyncCounter;
+} TentacleDevice;
+
+/**
+ Returns a device struct with default values.
+ 
+ @return TentacleDevice struct with default values.
+ */
+TentacleDevice TentacleDeviceMake(void);
+    
+/**
+ Update precise, calculated timecode for this device.
+ 
+ @param device Update cached timecodes of this device.
+ @param advertisementNew Add timecode from this newly received advertisement
+ to cache.
+ @return Difference between elapsed time according to received timecodes and
+ local device time, in milliseconds.
+ */
+double TentacleDeviceUpdateTimecode(TentacleDevice *device, const TentacleAdvertisement *advertisementNew);
+
+/**
+ If a device has not been seen after `TENTACLE_DEVICE_DISAPPEARED_THRESHOLD` seconds
+ it is considered disappeared.
+ 
+ @param device Device to get disappeared status for.
+ @param atTimestamp Get disappeared status for this timestamp, i.e. the current
+ time since boot in seconds (see TentacleTimecode.receivedTimestamp for specific timestamp-source).
+ @return True if device has not been seen for `TENTACLE_DEVICE_DISAPPEARED_THRESHOLD`
+ seconds.
+ */
+bool TentacleDeviceIsDisappeared(TentacleDevice *device, double atTimestamp);
+
+/**
+ If a device has not been seen after `TENTACLE_DEVICE_UNAVAILABLE_THRESHOLD` seconds
+ it is considered unavailable.
+
+ @param device Device to get unavailable status for.
+ @param atTimestamp Get unavailable status for this timestamp, i.e. the current
+ time since boot in seconds (see TentacleTimecode.receivedTimestamp for specific timestamp-source).
+ @return True if device has not been seen for `TENTACLE_DEVICE_UNAVAILABLE_THRESHOLD`
+ seconds.
+ */
+bool TentacleDeviceIsUnavailable(TentacleDevice *device, double atTimestamp);
+
+/**
+ Add a received signal strength indicator (RSSI) value to the ring buffer.
+
+ @param device Add to ring buffer of this device.
+ @param rssi Add this RSSI value to ring buffer.
+ */
+void TentacleDeviceAddRssi(TentacleDevice *device, int rssi);
+
+/**
+ Averaged received signal strength indicator (RSSI) value.
+ 
+ @param device Get averaged RSSI value for this device.
+ @return Averaged RSSI value
+ */
+int TentacleDeviceGetRssi(TentacleDevice *device);
+    
+#ifdef __cplusplus
+}
+#endif
+
+#endif /* tentacledevice_h */
diff --git a/Engine/Extras/VirtualProduction/LiveLinkVCAM/Tentacle.xcframework/ios-arm64_i386_x86_64-simulator/Tentacle.framework/Headers/TentacleDeviceCache.h b/Engine/Extras/VirtualProduction/LiveLinkVCAM/Tentacle.xcframework/ios-arm64_i386_x86_64-simulator/Tentacle.framework/Headers/TentacleDeviceCache.h
new file mode 100644
--- /dev/null
+++ b/Engine/Extras/VirtualProduction/LiveLinkVCAM/Tentacle.xcframework/ios-arm64_i386_x86_64-simulator/Tentacle.framework/Headers/TentacleDeviceCache.h
@@ -0,0 +1,136 @@
+//
+//  tentacleDeviceCache.h
+//  Tentacle-SDK
+//
+//  Copyright 2019 Tentacle Sync GmbH
+//
+//  Redistribution and use in source and binary forms, with or without
+//  modification, are permitted provided that the following conditions are met:
+//
+//  1. Redistributions of source code must retain the above copyright notice,
+//     this list of conditions and the following disclaimer.
+//
+//  2. Redistributions in binary form must reproduce the above copyright notice,
+//     this list of conditions and the following disclaimer in the
+//     documentation and/or other materials provided with the distribution.
+//
+//  3. Neither the name of the copyright holder nor the names of its
+//     contributors may be used to endorse or promote products derived from this
+//     software without specific prior written permission.
+//
+//  THIS SOFTWARE IS PROVIDED BY THE COPYRIGHT HOLDERS AND CONTRIBUTORS "AS IS"
+//  AND ANY EXPRESS OR IMPLIED WARRANTIES, INCLUDING, BUT NOT LIMITED TO, THE
+//  IMPLIED WARRANTIES OF MERCHANTABILITY AND FITNESS FOR A PARTICULAR PURPOSE
+//  ARE DISCLAIMED. IN NO EVENT SHALL THE COPYRIGHT HOLDER OR CONTRIBUTORS BE
+//  LIABLE FOR ANY DIRECT, INDIRECT, INCIDENTAL, SPECIAL, EXEMPLARY, OR
+//  CONSEQUENTIAL DAMAGES (INCLUDING, BUT NOT LIMITED TO, PROCUREMENT OF
+//  SUBSTITUTE GOODS OR SERVICES; LOSS OF USE, DATA, OR PROFITS; OR BUSINESS
+//  INTERRUPTION) HOWEVER CAUSED AND ON ANY THEORY OF LIABILITY, WHETHER IN
+//  CONTRACT, STRICT LIABILITY, OR TORT (INCLUDING NEGLIGENCE OR OTHERWISE)
+//  ARISING IN ANY WAY OUT OF THE USE OF THIS SOFTWARE, EVEN IF ADVISED OF THE
+//  POSSIBILITY OF SUCH DAMAGE.
+//
+
+/**
+ @file
+ @brief Discovered devices.
+ */
+
+#ifndef tentacleDeviceCache_h
+#define tentacleDeviceCache_h
+
+#ifdef __cplusplus
+extern "C" {
+#endif
+
+#include <stdio.h>
+#include <Tentacle/TentacleAdvertisement.h>
+#include <Tentacle/TentacleDevice.h>
+
+/**
+ Process a received Bluetooth advertisement packet.
+
+ @param advertisement Advertisement data to be processed.
+ @return Index of updated device in device cache.
+ */
+size_t TentacleDeviceCacheProcess(const TentacleAdvertisement *advertisement);
+    
+/**
+ Get current size of device cache.
+ 
+ @return Size of device cache.
+ */
+size_t TentacleDeviceCacheGetSize(void);
+    
+/**
+ Get current size of timecode cache for a specific device index.
+ 
+ @param deviceIndex Get cache size for device with this index.
+ @return Timecode cache size.
+ */
+size_t TentacleDeviceCacheGetTimecodeCacheSize(int deviceIndex);
+
+/**
+ Get cached device at specific index.
+ 
+ @param deviceIndex Index of device to get.
+ @return Device struct at specific index.
+ */
+TentacleDevice TentacleDeviceCacheGetDevice(int deviceIndex);
+
+/**
+ Get cached advertisement for device at specific index.
+
+ @param deviceIndex Index of device to get advertisement for.
+ @return Cached advertisement struct.
+ */
+TentacleAdvertisement TentacleDeviceCacheGetAdvertisement(int deviceIndex);
+   
+/**
+ Get timecode string representing a cached timecode.
+ 
+ @param deviceIndex Get timecode string for device with this index.
+ @param timecodeIndex Get timecode string for device with this index.
+ @param timecodeCacheString Resulting string value. Has to be of length
+        `TENTACLE_TIMECODE_STRING_LENGTH`, includes terminating null character.
+ @return Number of characters of `timecodeCacheString`, not including the
+         terminating null character.
+ */
+int TentacleDeviceCacheTimecodeString(int deviceIndex, int timecodeIndex, char *timecodeCacheString);
+    
+/**
+ Get timecode string for a specific point in time.
+ 
+ @param deviceIndex Get timecode string for device with this index.
+ @param atTimestamp Get timecode at this timestamp.
+ @param timecodeString Resulting string value. Has to be of length
+        `TENTACLE_TIMECODE_STRING_LENGTH`.
+ @return Number of characters of `timecodeString`, not including the
+         terminating null character.
+ */
+int TentacleDeviceCacheTimecodeStringAtTimestamp(int deviceIndex, double atTimestamp, char *timecodeString);
+
+/**
+ Get timecode string for a specific point in time.
+
+ @param identifier Identifier of device to get.
+ @param identifierLength Length of identifier.
+ @param atTimestamp Get timecode at this timestamp.
+ @param timecodeString Resulting string value. Has to be of length `TENTACLE_TIMECODE_STRING_LENGTH`.
+ @return Number of characters of `timecodeString`, not including the
+ terminating null character.
+ */
+int TentacleDeviceCacheTimecodeStringAtTimestampForIdentifier(char *identifier, size_t identifierLength, double atTimestamp, char *timecodeString);
+
+/**
+ Print string representation of device cache to stdout.
+
+ @param atTimestamp Print timecodes at this timestamp.
+ */
+void TentacleDeviceCachePrint(double atTimestamp);
+    
+#ifdef __cplusplus
+}
+#endif
+
+#endif /* tentacleDeviceCache_h */
diff --git a/Engine/Extras/VirtualProduction/LiveLinkVCAM/Tentacle.xcframework/ios-arm64_i386_x86_64-simulator/Tentacle.framework/Headers/TentacleEnum.h b/Engine/Extras/VirtualProduction/LiveLinkVCAM/Tentacle.xcframework/ios-arm64_i386_x86_64-simulator/Tentacle.framework/Headers/TentacleEnum.h
new file mode 100644
--- /dev/null
+++ b/Engine/Extras/VirtualProduction/LiveLinkVCAM/Tentacle.xcframework/ios-arm64_i386_x86_64-simulator/Tentacle.framework/Headers/TentacleEnum.h
@@ -0,0 +1,92 @@
+//
+//  tentacleEnum.h
+//  Tentacle-SDK
+//
+//  Copyright 2019 Tentacle Sync GmbH
+//
+//  Redistribution and use in source and binary forms, with or without
+//  modification, are permitted provided that the following conditions are met:
+//
+//  1. Redistributions of source code must retain the above copyright notice,
+//     this list of conditions and the following disclaimer.
+//
+//  2. Redistributions in binary form must reproduce the above copyright notice,
+//     this list of conditions and the following disclaimer in the
+//     documentation and/or other materials provided with the distribution.
+//
+//  3. Neither the name of the copyright holder nor the names of its
+//     contributors may be used to endorse or promote products derived from this
+//     software without specific prior written permission.
+//
+//  THIS SOFTWARE IS PROVIDED BY THE COPYRIGHT HOLDERS AND CONTRIBUTORS "AS IS"
+//  AND ANY EXPRESS OR IMPLIED WARRANTIES, INCLUDING, BUT NOT LIMITED TO, THE
+//  IMPLIED WARRANTIES OF MERCHANTABILITY AND FITNESS FOR A PARTICULAR PURPOSE
+//  ARE DISCLAIMED. IN NO EVENT SHALL THE COPYRIGHT HOLDER OR CONTRIBUTORS BE
+//  LIABLE FOR ANY DIRECT, INDIRECT, INCIDENTAL, SPECIAL, EXEMPLARY, OR
+//  CONSEQUENTIAL DAMAGES (INCLUDING, BUT NOT LIMITED TO, PROCUREMENT OF
+//  SUBSTITUTE GOODS OR SERVICES; LOSS OF USE, DATA, OR PROFITS; OR BUSINESS
+//  INTERRUPTION) HOWEVER CAUSED AND ON ANY THEORY OF LIABILITY, WHETHER IN
+//  CONTRACT, STRICT LIABILITY, OR TORT (INCLUDING NEGLIGENCE OR OTHERWISE)
+//  ARISING IN ANY WAY OUT OF THE USE OF THIS SOFTWARE, EVEN IF ADVISED OF THE
+//  POSSIBILITY OF SUCH DAMAGE.
+//
+
+/**
+ @file
+ @brief Tentacle enumerations.
+ */
+
+#ifndef tentacleenum_h
+#define tentacleenum_h
+
+#ifdef TentacleAdvertisement
+extern "C" {
+#endif
+    
+/**
+ State the device is in.
+ */
+enum tentacleDeviceMode {
+    /// Device is powered off.
+    power_off = 0,
+    
+    /// Device is booting.
+    startup = 1,
+    
+    /// Device is in standby.
+    standby = 2,
+    
+    /// No cable is plugged into the device, no audio timecode is generated.
+    unplugged = 3,
+    
+    /// Cable is plugged in and device is generating timecode.
+    generateTimecode = 4,
+    
+    /// Device is reading timecode.
+    readTimecode = 5,
+    
+    /// Device is stopped.
+    stopped = 6
+};
+
+/**
+ When producing timecode to be recorded on an audio track, the timecode is
+ encoded as an audio signal using a certain volume. This volume can be
+ low (Mic), high (Line), something in between (Custom) or be automatically
+ detected (AutoMic).
+ */
+enum TentacleAudioLevel {
+    /// Audio timecode is produced at a low volume
+    mic = 0,
+    
+    /// Audio timecode is produced at a high volume
+    line = 1 ,
+    
+    /// Audio timecode is produced at a medium volume
+    custom = 2,
+    
+    /// Audio timecode is automatically set to a low volume
+    autoMic = 3
+};
+
+#endif /* tentacleenum_h */
diff --git a/Engine/Extras/VirtualProduction/LiveLinkVCAM/Tentacle.xcframework/ios-arm64_i386_x86_64-simulator/Tentacle.framework/Headers/TentacleIcon.h b/Engine/Extras/VirtualProduction/LiveLinkVCAM/Tentacle.xcframework/ios-arm64_i386_x86_64-simulator/Tentacle.framework/Headers/TentacleIcon.h
new file mode 100644
--- /dev/null
+++ b/Engine/Extras/VirtualProduction/LiveLinkVCAM/Tentacle.xcframework/ios-arm64_i386_x86_64-simulator/Tentacle.framework/Headers/TentacleIcon.h
@@ -0,0 +1,80 @@
+//
+//  TentacleIcon.h
+//  Tentacle-SDK
+//
+//  Copyright 2019 Tentacle Sync GmbH
+//
+//  Redistribution and use in source and binary forms, with or without
+//  modification, are permitted provided that the following conditions are met:
+//
+//  1. Redistributions of source code must retain the above copyright notice,
+//     this list of conditions and the following disclaimer.
+//
+//  2. Redistributions in binary form must reproduce the above copyright notice,
+//     this list of conditions and the following disclaimer in the
+//     documentation and/or other materials provided with the distribution.
+//
+//  3. Neither the name of the copyright holder nor the names of its
+//     contributors may be used to endorse or promote products derived from this
+//     software without specific prior written permission.
+//
+//  THIS SOFTWARE IS PROVIDED BY THE COPYRIGHT HOLDERS AND CONTRIBUTORS "AS IS"
+//  AND ANY EXPRESS OR IMPLIED WARRANTIES, INCLUDING, BUT NOT LIMITED TO, THE
+//  IMPLIED WARRANTIES OF MERCHANTABILITY AND FITNESS FOR A PARTICULAR PURPOSE
+//  ARE DISCLAIMED. IN NO EVENT SHALL THE COPYRIGHT HOLDER OR CONTRIBUTORS BE
+//  LIABLE FOR ANY DIRECT, INDIRECT, INCIDENTAL, SPECIAL, EXEMPLARY, OR
+//  CONSEQUENTIAL DAMAGES (INCLUDING, BUT NOT LIMITED TO, PROCUREMENT OF
+//  SUBSTITUTE GOODS OR SERVICES; LOSS OF USE, DATA, OR PROFITS; OR BUSINESS
+//  INTERRUPTION) HOWEVER CAUSED AND ON ANY THEORY OF LIABILITY, WHETHER IN
+//  CONTRACT, STRICT LIABILITY, OR TORT (INCLUDING NEGLIGENCE OR OTHERWISE)
+//  ARISING IN ANY WAY OUT OF THE USE OF THIS SOFTWARE, EVEN IF ADVISED OF THE
+//  POSSIBILITY OF SUCH DAMAGE.
+//
+
+@import Foundation;
+
+#if TARGET_OS_IOS
+@import UIKit;
+#endif
+
+NS_ASSUME_NONNULL_BEGIN
+
+@interface TentacleIcon: NSObject
+
+#if TARGET_OS_IOS
+
+/// Get Tentacle device icon image.
+/// \param index Icon index of range 0 to 37
+///
+/// returns:
+/// Device image
++ (UIImage * _Nonnull)tentacleIconWithIndex:(NSUInteger)index;
+
+/// Get default battery image.
+/// \param percent Percentage of battery status
+///
+/// returns:
+/// Battery image
++ (UIImage * _Nonnull)tentacleBatteryIconForPercent:(NSUInteger)percent;
+
+/// Get default signal strength image.
+/// \param index Possible values are 0 to 3
+///
+/// returns:
+/// Signal strength image
++ (UIImage * _Nonnull)tentacleSignalStrengthIconWithIndex:(NSUInteger)index;
+
+/// Get default signal strength image for a given received signal strength indicator (RSSI) value.
+/// \param rssiValue Get image corresponding to this RSSI
+/// \param linkMode Internal use only, always pass false
+///
+/// returns:
+/// Signal strength image
++ (UIImage * _Nonnull)tentacleSignalStrengthIconForRssiValue:(NSInteger)rssiValue linkMode:(BOOL)linkMode;
+
+#endif // #if TARGET_OS_IOS
+
+@end
+
+NS_ASSUME_NONNULL_END
+
diff --git a/Engine/Extras/VirtualProduction/LiveLinkVCAM/Tentacle.xcframework/ios-arm64_i386_x86_64-simulator/Tentacle.framework/Headers/TentacleLog.h b/Engine/Extras/VirtualProduction/LiveLinkVCAM/Tentacle.xcframework/ios-arm64_i386_x86_64-simulator/Tentacle.framework/Headers/TentacleLog.h
new file mode 100644
--- /dev/null
+++ b/Engine/Extras/VirtualProduction/LiveLinkVCAM/Tentacle.xcframework/ios-arm64_i386_x86_64-simulator/Tentacle.framework/Headers/TentacleLog.h
@@ -0,0 +1,66 @@
+//
+//  tentacleLog.h
+//  Tentacle-SDK
+//
+//  Copyright 2019 Tentacle Sync GmbH
+//
+//  Redistribution and use in source and binary forms, with or without
+//  modification, are permitted provided that the following conditions are met:
+//
+//  1. Redistributions of source code must retain the above copyright notice,
+//     this list of conditions and the following disclaimer.
+//
+//  2. Redistributions in binary form must reproduce the above copyright notice,
+//     this list of conditions and the following disclaimer in the
+//     documentation and/or other materials provided with the distribution.
+//
+//  3. Neither the name of the copyright holder nor the names of its
+//     contributors may be used to endorse or promote products derived from this
+//     software without specific prior written permission.
+//
+//  THIS SOFTWARE IS PROVIDED BY THE COPYRIGHT HOLDERS AND CONTRIBUTORS "AS IS"
+//  AND ANY EXPRESS OR IMPLIED WARRANTIES, INCLUDING, BUT NOT LIMITED TO, THE
+//  IMPLIED WARRANTIES OF MERCHANTABILITY AND FITNESS FOR A PARTICULAR PURPOSE
+//  ARE DISCLAIMED. IN NO EVENT SHALL THE COPYRIGHT HOLDER OR CONTRIBUTORS BE
+//  LIABLE FOR ANY DIRECT, INDIRECT, INCIDENTAL, SPECIAL, EXEMPLARY, OR
+//  CONSEQUENTIAL DAMAGES (INCLUDING, BUT NOT LIMITED TO, PROCUREMENT OF
+//  SUBSTITUTE GOODS OR SERVICES; LOSS OF USE, DATA, OR PROFITS; OR BUSINESS
+//  INTERRUPTION) HOWEVER CAUSED AND ON ANY THEORY OF LIABILITY, WHETHER IN
+//  CONTRACT, STRICT LIABILITY, OR TORT (INCLUDING NEGLIGENCE OR OTHERWISE)
+//  ARISING IN ANY WAY OUT OF THE USE OF THIS SOFTWARE, EVEN IF ADVISED OF THE
+//  POSSIBILITY OF SUCH DAMAGE.
+//
+
+/**
+ @file
+ @brief Discovered devices.
+ */
+
+#ifndef tentaclelog_h
+#define tentaclelog_h
+
+#ifdef __cplusplus
+extern "C"
+{
+#endif
+
+#include <stdio.h>
+#include <stdarg.h>
+
+#if __ANDROID__ == 1
+#include <android/log.h>
+#endif
+
+/**
+ * Prints logging information using platform specific methods.
+ *
+ * @param fmt Format string.
+ * @param ... Variadic arguments to be logged.
+ */
+void TentacleLog(const char *fmt, ...);
+
+#ifdef __cplusplus
+}
+#endif
+
+#endif /* tentaclelog_h */
diff --git a/Engine/Extras/VirtualProduction/LiveLinkVCAM/Tentacle.xcframework/ios-arm64_i386_x86_64-simulator/Tentacle.framework/Headers/TentacleTimecode.h b/Engine/Extras/VirtualProduction/LiveLinkVCAM/Tentacle.xcframework/ios-arm64_i386_x86_64-simulator/Tentacle.framework/Headers/TentacleTimecode.h
new file mode 100644
--- /dev/null
+++ b/Engine/Extras/VirtualProduction/LiveLinkVCAM/Tentacle.xcframework/ios-arm64_i386_x86_64-simulator/Tentacle.framework/Headers/TentacleTimecode.h
@@ -0,0 +1,172 @@
+//
+//  tentacleTimecode.h
+//  Tentacle-SDK
+//
+//  Copyright 2019 Tentacle Sync GmbH
+//
+//  Redistribution and use in source and binary forms, with or without
+//  modification, are permitted provided that the following conditions are met:
+//
+//  1. Redistributions of source code must retain the above copyright notice,
+//     this list of conditions and the following disclaimer.
+//
+//  2. Redistributions in binary form must reproduce the above copyright notice,
+//     this list of conditions and the following disclaimer in the
+//     documentation and/or other materials provided with the distribution.
+//
+//  3. Neither the name of the copyright holder nor the names of its
+//     contributors may be used to endorse or promote products derived from this
+//     software without specific prior written permission.
+//
+//  THIS SOFTWARE IS PROVIDED BY THE COPYRIGHT HOLDERS AND CONTRIBUTORS "AS IS"
+//  AND ANY EXPRESS OR IMPLIED WARRANTIES, INCLUDING, BUT NOT LIMITED TO, THE
+//  IMPLIED WARRANTIES OF MERCHANTABILITY AND FITNESS FOR A PARTICULAR PURPOSE
+//  ARE DISCLAIMED. IN NO EVENT SHALL THE COPYRIGHT HOLDER OR CONTRIBUTORS BE
+//  LIABLE FOR ANY DIRECT, INDIRECT, INCIDENTAL, SPECIAL, EXEMPLARY, OR
+//  CONSEQUENTIAL DAMAGES (INCLUDING, BUT NOT LIMITED TO, PROCUREMENT OF
+//  SUBSTITUTE GOODS OR SERVICES; LOSS OF USE, DATA, OR PROFITS; OR BUSINESS
+//  INTERRUPTION) HOWEVER CAUSED AND ON ANY THEORY OF LIABILITY, WHETHER IN
+//  CONTRACT, STRICT LIABILITY, OR TORT (INCLUDING NEGLIGENCE OR OTHERWISE)
+//  ARISING IN ANY WAY OUT OF THE USE OF THIS SOFTWARE, EVEN IF ADVISED OF THE
+//  POSSIBILITY OF SUCH DAMAGE.
+//
+
+/**
+ @file
+ @brief Tentacle timecode.
+ 
+ Contains cached raw data of a received timecode as received from a Bluetooth
+ advertisement.
+ */
+
+#ifndef tentacletimecode_h
+#define tentacletimecode_h
+
+#ifdef __cplusplus
+extern "C" {
+#endif
+
+#include <stdio.h>   // for size_t
+#include <stdint.h>  // for uint8_t
+#include <stdbool.h> // for bool
+    
+/// Length of string returned by `tentacleTimecodeCacheGetString`.
+#define TENTACLE_TIMECODE_STRING_LENGTH 35
+    
+/// 23.98 FPS =  24 FPS NTSC
+#define TENTACLE_FPS_23_98 (24000.0/1001.0)
+    
+/// 29.97 FPS = 30 FPS NTSC
+#define TENTACLE_FPS_29_97 (30000.0/1001.0)
+
+/**
+ Contains cached raw data of a received timecode as received from a Bluetooth
+ advertisement.
+ */
+typedef struct TentacleTimecode {
+    /// Hour of timecode.
+    uint8_t hour;
+    
+    /// Minute of timecode.
+    uint8_t minute;
+    
+    /// Second of timecode.
+    uint8_t second;
+    
+    /// Frame of timecode.
+    uint8_t frame;
+    
+    /// Microseconds since start of last frame.
+    uint16_t microsecond;
+    
+    /// \brief Timestamp of when this timecode has been received, in seconds.
+    ///
+    /// On Apple Darwin, [machAbsoluteTime]([REDACTED_PROFILE_URL]/darwin-xnu/blob/0a798f6738bc1db01281fc08ae024145e84df927/osfmk/arm/rtclock.c#L162) is used.
+    ///
+    /// On Android, [ScanResult.getTimestampNanos()](https://developer.android.com/reference/android/bluetooth/le/ScanResult.html#getTimestampNanos()) is used.
+    ///
+    /// Caution: Due to maximum precision these Timestamps are based on the boot-time of the device. According to this, the timestamps are invalid after
+    /// a device reboot and should not be used for long time comparison or alike.
+    double receivedTimestamp;
+} TentacleTimecode;
+    
+/**
+ Returns a timecode struct with default values.
+
+ @return TentacleTimecode struct with default values.
+ */
+TentacleTimecode TentacleTimecodeMake(void);
+    
+/**
+ Check if two timecodes are equal.
+
+ @param timecode1 First timecode to compare.
+ @param timecode2 Second timecode to compare.
+ @return True if timecodes are equal, false otherwise.
+ */
+bool TentacleTimecodeEqual(const TentacleTimecode *timecode1, const TentacleTimecode *timecode2);
+
+/**
+ Total number of seconds described by this timecode at specific timestamp.
+
+ @param timecode Get total number of seconds for this timecode.
+ @param frameRate Frame rate used to calculate number of seconds.
+ @param dropFrame True if timecode is using drop frames.
+ @param atTimestamp Total number of seconds of timecode at this timestamp.
+ @return Frame rate used to calculate duration of frame count at timestamp.
+ */
+double TentacleTimecodeSecondsAtTimestamp(const TentacleTimecode *timecode, double frameRate, bool dropFrame, double atTimestamp);
+    
+/**
+ Create a timecode from a total number of seconds.
+
+ @param seconds Convert this total number of seconds to timecode.
+ @param frameRate Frame rate used during conversion.
+ @param dropFrame True if timecode is using drop frames.
+ @param receivedTimestamp Use this as received timestamp.
+ @return New timecode struct containing the passed total number of seconds.
+ */
+TentacleTimecode TentacleTimecodeFromSeconds(double seconds, double frameRate, bool dropFrame, double receivedTimestamp);
+
+/**
+ Given a received timecode, calculate what its value would be at a different
+ point in time.
+
+ @param timecode Get this timecode for a specific point in time
+ @param frameRate Frame rate of timecode
+ @param dropFrame True if timecode is using drop frames.
+ @param atTimestamp Requested point in time
+ @return Timecode value at specific point in time.
+ */
+TentacleTimecode TentacleTimecodeAtTimestamp(TentacleTimecode timecode, double frameRate, bool dropFrame, double atTimestamp);
+
+/**
+ Get timecode string representing a timecode.
+ 
+ @param timecode Get string value representing this timecode struct.
+ @param dropFrame True if timecode is drop frame.
+ @param timecodeCacheString Resulting string value. Has to be of length
+        `TENTACLE_TIMECODE_STRING_LENGTH`.
+ @return Number of characters of `timecodeCacheString`, not including the
+         terminating null character.
+ */
+int TentacleTimecodeString(TentacleTimecode timecode, bool dropFrame, char *timecodeCacheString);
+
+/**
+ Get timecode string representing a timecode, including microseconds since last
+ frame and received timestamp.
+ 
+ @param timecode Get string value representing this timecode struct.
+ @param dropFrame True if timecode is drop frame.
+ @param timecodeCacheString Resulting string value. Has to be of length
+        `TENTACLE_TIMECODE_STRING_LENGTH`.
+ @return Number of characters of `timecodeCacheString`, not including the
+         terminating null character.
+ */
+int TentacleTimecodeStringVerbose(TentacleTimecode timecode, bool dropFrame, char *timecodeCacheString);
+
+#ifdef __cplusplus
+}
+#endif
+
+#endif /* tentacletimecode_h */
