Live Link VCAM : (re)initial checkin : the Tentacle Framework was causing horde failures as it is a 3rd party library.

diff --git a/Engine/Extras/VirtualProduction/LiveLinkVCAM/LiveLinkProvider.framework/Headers/LiveLink.h b/Engine/Extras/VirtualProduction/LiveLinkVCAM/LiveLinkProvider.framework/Headers/LiveLink.h
new file mode 100644
--- /dev/null
+++ b/Engine/Extras/VirtualProduction/LiveLinkVCAM/LiveLinkProvider.framework/Headers/LiveLink.h
@@ -0,0 +1,21 @@
+
+#import <Foundation/Foundation.h>
+
+#import "LiveLinkProvider.h"
+
+__attribute__((visibility("default")))
+@protocol UELogDelegate
+
+- (void) ueLogMessage:(NSString *)message;
+
+@end
+
+ 
+__attribute__((visibility("default")))
+@interface LiveLink : NSObject
+
++ (void) initialize:(id<UELogDelegate>)logDelegate;
++ (void) shutdown;
++ (id<LiveLinkProvider> ) createProvider:(NSString *)name;
+
+@end
diff --git a/Engine/Extras/VirtualProduction/LiveLinkVCAM/LiveLinkProvider.framework/Headers/LiveLinkProvider.h b/Engine/Extras/VirtualProduction/LiveLinkVCAM/LiveLinkProvider.framework/Headers/LiveLinkProvider.h
new file mode 100644
--- /dev/null
+++ b/Engine/Extras/VirtualProduction/LiveLinkVCAM/LiveLinkProvider.framework/Headers/LiveLinkProvider.h
@@ -0,0 +1,14 @@
+
+#import <Foundation/Foundation.h>
+#import <simd/simd.h>
+
+@protocol LiveLinkProvider <NSObject>
+
+- (void)addCameraSubject:(NSString *)subjectName;
+- (void)removeCameraSubject:(NSString *)subjectName;
+
+- (void)updateSubject:(NSString *)subjectName withTransform:(simd_float4x4)xform;
+- (void)updateSubject:(NSString *)subjectName withTransform:(simd_float4x4)xform atTime:(double)time;
+
+@end
+
diff --git a/Engine/Extras/VirtualProduction/LiveLinkVCAM/README.txt b/Engine/Extras/VirtualProduction/LiveLinkVCAM/README.txt
new file mode 100644
--- /dev/null
+++ b/Engine/Extras/VirtualProduction/LiveLinkVCAM/README.txt
@@ -0,0 +1,8 @@
+Live Link VCAM 
+==============
+
+To build this app, the Tentacle.xcframework is required. For the moment, see Brian Smith or Ryan Mayeda for access.
+
+
+
+
diff --git a/Engine/Extras/VirtualProduction/LiveLinkVCAM/vcam/Assets.xcassets/AccentColor.colorset/Contents.json b/Engine/Extras/VirtualProduction/LiveLinkVCAM/vcam/Assets.xcassets/AccentColor.colorset/Contents.json
new file mode 100644
--- /dev/null
+++ b/Engine/Extras/VirtualProduction/LiveLinkVCAM/vcam/Assets.xcassets/AccentColor.colorset/Contents.json
@@ -0,0 +1,28 @@
+{
+  "colors" : [
+    {
+      "color" : {
+        "platform" : "ios",
+        "reference" : "systemBlueColor"
+      },
+      "idiom" : "universal"
+    },
+    {
+      "appearances" : [
+        {
+          "appearance" : "luminosity",
+          "value" : "dark"
+        }
+      ],
+      "color" : {
+        "platform" : "ios",
+        "reference" : "systemBlueColor"
+      },
+      "idiom" : "universal"
+    }
+  ],
+  "info" : {
+    "author" : "xcode",
+    "version" : 1
+  }
+}
diff --git a/Engine/Extras/VirtualProduction/LiveLinkVCAM/vcam/Assets.xcassets/AppIcon.appiconset/Contents.json b/Engine/Extras/VirtualProduction/LiveLinkVCAM/vcam/Assets.xcassets/AppIcon.appiconset/Contents.json
new file mode 100644
--- /dev/null
+++ b/Engine/Extras/VirtualProduction/LiveLinkVCAM/vcam/Assets.xcassets/AppIcon.appiconset/Contents.json
@@ -0,0 +1,116 @@
+{
+  "images" : [
+    {
+      "filename" : "[REDACTED_EMAIL]",
+      "idiom" : "iphone",
+      "scale" : "2x",
+      "size" : "20x20"
+    },
+    {
+      "filename" : "[REDACTED_EMAIL]",
+      "idiom" : "iphone",
+      "scale" : "3x",
+      "size" : "20x20"
+    },
+    {
+      "filename" : "[REDACTED_EMAIL]",
+      "idiom" : "iphone",
+      "scale" : "2x",
+      "size" : "29x29"
+    },
+    {
+      "filename" : "[REDACTED_EMAIL]",
+      "idiom" : "iphone",
+      "scale" : "3x",
+      "size" : "29x29"
+    },
+    {
+      "filename" : "[REDACTED_EMAIL]",
+      "idiom" : "iphone",
+      "scale" : "2x",
+      "size" : "40x40"
+    },
+    {
+      "filename" : "[REDACTED_EMAIL]",
+      "idiom" : "iphone",
+      "scale" : "3x",
+      "size" : "40x40"
+    },
+    {
+      "filename" : "[REDACTED_EMAIL]",
+      "idiom" : "iphone",
+      "scale" : "2x",
+      "size" : "60x60"
+    },
+    {
+      "filename" : "[REDACTED_EMAIL]",
+      "idiom" : "iphone",
+      "scale" : "3x",
+      "size" : "60x60"
+    },
+    {
+      "filename" : "[REDACTED_EMAIL]",
+      "idiom" : "ipad",
+      "scale" : "1x",
+      "size" : "20x20"
+    },
+    {
+      "filename" : "[REDACTED_EMAIL]",
+      "idiom" : "ipad",
+      "scale" : "2x",
+      "size" : "20x20"
+    },
+    {
+      "filename" : "[REDACTED_EMAIL]",
+      "idiom" : "ipad",
+      "scale" : "1x",
+      "size" : "29x29"
+    },
+    {
+      "filename" : "[REDACTED_EMAIL]",
+      "idiom" : "ipad",
+      "scale" : "2x",
+      "size" : "29x29"
+    },
+    {
+      "filename" : "[REDACTED_EMAIL]",
+      "idiom" : "ipad",
+      "scale" : "1x",
+      "size" : "40x40"
+    },
+    {
+      "filename" : "[REDACTED_EMAIL]",
+      "idiom" : "ipad",
+      "scale" : "2x",
+      "size" : "40x40"
+    },
+    {
+      "filename" : "[REDACTED_EMAIL]",
+      "idiom" : "ipad",
+      "scale" : "1x",
+      "size" : "76x76"
+    },
+    {
+      "filename" : "[REDACTED_EMAIL]",
+      "idiom" : "ipad",
+      "scale" : "2x",
+      "size" : "76x76"
+    },
+    {
+      "filename" : "[REDACTED_EMAIL]",
+      "idiom" : "ipad",
+      "scale" : "2x",
+      "size" : "83.5x83.5"
+    },
+    {
+      "filename" : "Icon-App-1024x1024.png",
+      "idiom" : "ios-marketing",
+      "scale" : "1x",
+      "size" : "1024x1024"
+    }
+  ],
+  "info" : {
+    "author" : "xcode",
+    "version" : 1
+  }
+}
diff --git a/Engine/Extras/VirtualProduction/LiveLinkVCAM/vcam/Assets.xcassets/BG.imageset/Contents.json b/Engine/Extras/VirtualProduction/LiveLinkVCAM/vcam/Assets.xcassets/BG.imageset/Contents.json
new file mode 100644
--- /dev/null
+++ b/Engine/Extras/VirtualProduction/LiveLinkVCAM/vcam/Assets.xcassets/BG.imageset/Contents.json
@@ -0,0 +1,23 @@
+{
+  "images" : [
+    {
+      "filename" : "BG [REDACTED_EMAIL]",
+      "idiom" : "universal",
+      "scale" : "1x"
+    },
+    {
+      "filename" : "BG [REDACTED_EMAIL]",
+      "idiom" : "universal",
+      "scale" : "2x"
+    },
+    {
+      "filename" : "BG [REDACTED_EMAIL]",
+      "idiom" : "universal",
+      "scale" : "3x"
+    }
+  ],
+  "info" : {
+    "author" : "xcode",
+    "version" : 1
+  }
+}
diff --git a/Engine/Extras/VirtualProduction/LiveLinkVCAM/vcam/Assets.xcassets/Contents.json b/Engine/Extras/VirtualProduction/LiveLinkVCAM/vcam/Assets.xcassets/Contents.json
new file mode 100644
--- /dev/null
+++ b/Engine/Extras/VirtualProduction/LiveLinkVCAM/vcam/Assets.xcassets/Contents.json
@@ -0,0 +1,6 @@
+{
+  "info" : {
+    "author" : "xcode",
+    "version" : 1
+  }
+}
diff --git a/Engine/Extras/VirtualProduction/LiveLinkVCAM/vcam/Assets.xcassets/Exit.imageset/Contents.json b/Engine/Extras/VirtualProduction/LiveLinkVCAM/vcam/Assets.xcassets/Exit.imageset/Contents.json
new file mode 100644
--- /dev/null
+++ b/Engine/Extras/VirtualProduction/LiveLinkVCAM/vcam/Assets.xcassets/Exit.imageset/Contents.json
@@ -0,0 +1,26 @@
+{
+  "images" : [
+    {
+      "filename" : "602-exit.png",
+      "idiom" : "universal",
+      "scale" : "1x"
+    },
+    {
+      "filename" : "[REDACTED_EMAIL]",
+      "idiom" : "universal",
+      "scale" : "2x"
+    },
+    {
+      "filename" : "[REDACTED_EMAIL]",
+      "idiom" : "universal",
+      "scale" : "3x"
+    }
+  ],
+  "info" : {
+    "author" : "xcode",
+    "version" : 1
+  },
+  "properties" : {
+    "template-rendering-intent" : "template"
+  }
+}
diff --git a/Engine/Extras/VirtualProduction/LiveLinkVCAM/vcam/Assets.xcassets/Logo.imageset/Contents.json b/Engine/Extras/VirtualProduction/LiveLinkVCAM/vcam/Assets.xcassets/Logo.imageset/Contents.json
new file mode 100644
--- /dev/null
+++ b/Engine/Extras/VirtualProduction/LiveLinkVCAM/vcam/Assets.xcassets/Logo.imageset/Contents.json
@@ -0,0 +1,23 @@
+{
+  "images" : [
+    {
+      "filename" : "[REDACTED_EMAIL]",
+      "idiom" : "universal",
+      "scale" : "1x"
+    },
+    {
+      "filename" : "[REDACTED_EMAIL]",
+      "idiom" : "universal",
+      "scale" : "2x"
+    },
+    {
+      "filename" : "[REDACTED_EMAIL]",
+      "idiom" : "universal",
+      "scale" : "3x"
+    }
+  ],
+  "info" : {
+    "author" : "xcode",
+    "version" : 1
+  }
+}
diff --git a/Engine/Extras/VirtualProduction/LiveLinkVCAM/vcam/Assets.xcassets/Settings.imageset/Contents.json b/Engine/Extras/VirtualProduction/LiveLinkVCAM/vcam/Assets.xcassets/Settings.imageset/Contents.json
new file mode 100644
--- /dev/null
+++ b/Engine/Extras/VirtualProduction/LiveLinkVCAM/vcam/Assets.xcassets/Settings.imageset/Contents.json
@@ -0,0 +1,23 @@
+{
+  "images" : [
+    {
+      "filename" : "URFaceSettings.png",
+      "idiom" : "universal",
+      "scale" : "1x"
+    },
+    {
+      "filename" : "[REDACTED_EMAIL]",
+      "idiom" : "universal",
+      "scale" : "2x"
+    },
+    {
+      "filename" : "[REDACTED_EMAIL]",
+      "idiom" : "universal",
+      "scale" : "3x"
+    }
+  ],
+  "info" : {
+    "author" : "xcode",
+    "version" : 1
+  }
+}
diff --git a/Engine/Extras/VirtualProduction/LiveLinkVCAM/vcam/Assets.xcassets/checkerboard.imageset/Contents.json b/Engine/Extras/VirtualProduction/LiveLinkVCAM/vcam/Assets.xcassets/checkerboard.imageset/Contents.json
new file mode 100644
--- /dev/null
+++ b/Engine/Extras/VirtualProduction/LiveLinkVCAM/vcam/Assets.xcassets/checkerboard.imageset/Contents.json
@@ -0,0 +1,21 @@
+{
+  "images" : [
+    {
+      "filename" : "Checkerboard_tile.svg.png",
+      "idiom" : "universal",
+      "scale" : "1x"
+    },
+    {
+      "idiom" : "universal",
+      "scale" : "2x"
+    },
+    {
+      "idiom" : "universal",
+      "scale" : "3x"
+    }
+  ],
+  "info" : {
+    "author" : "xcode",
+    "version" : 1
+  }
+}
diff --git a/Engine/Extras/VirtualProduction/LiveLinkVCAM/vcam/Assets.xcassets/counterCloud.imageset/Contents.json b/Engine/Extras/VirtualProduction/LiveLinkVCAM/vcam/Assets.xcassets/counterCloud.imageset/Contents.json
new file mode 100644
--- /dev/null
+++ b/Engine/Extras/VirtualProduction/LiveLinkVCAM/vcam/Assets.xcassets/counterCloud.imageset/Contents.json
@@ -0,0 +1,23 @@
+{
+  "images" : [
+    {
+      "filename" : "CounterCloud.png",
+      "idiom" : "universal",
+      "scale" : "1x"
+    },
+    {
+      "filename" : "[REDACTED_EMAIL]",
+      "idiom" : "universal",
+      "scale" : "2x"
+    },
+    {
+      "filename" : "[REDACTED_EMAIL]",
+      "idiom" : "universal",
+      "scale" : "3x"
+    }
+  ],
+  "info" : {
+    "author" : "xcode",
+    "version" : 1
+  }
+}
diff --git a/Engine/Extras/VirtualProduction/LiveLinkVCAM/vcam/Assets.xcassets/counterDevice.imageset/Contents.json b/Engine/Extras/VirtualProduction/LiveLinkVCAM/vcam/Assets.xcassets/counterDevice.imageset/Contents.json
new file mode 100644
--- /dev/null
+++ b/Engine/Extras/VirtualProduction/LiveLinkVCAM/vcam/Assets.xcassets/counterDevice.imageset/Contents.json
@@ -0,0 +1,23 @@
+{
+  "images" : [
+    {
+      "filename" : "CounterDevice.png",
+      "idiom" : "universal",
+      "scale" : "1x"
+    },
+    {
+      "filename" : "[REDACTED_EMAIL]",
+      "idiom" : "universal",
+      "scale" : "2x"
+    },
+    {
+      "filename" : "[REDACTED_EMAIL]",
+      "idiom" : "universal",
+      "scale" : "3x"
+    }
+  ],
+  "info" : {
+    "author" : "xcode",
+    "version" : 1
+  }
+}
diff --git a/Engine/Extras/VirtualProduction/LiveLinkVCAM/vcam/Assets.xcassets/no_signal.textureset/Contents.json b/Engine/Extras/VirtualProduction/LiveLinkVCAM/vcam/Assets.xcassets/no_signal.textureset/Contents.json
new file mode 100644
--- /dev/null
+++ b/Engine/Extras/VirtualProduction/LiveLinkVCAM/vcam/Assets.xcassets/no_signal.textureset/Contents.json
@@ -0,0 +1,12 @@
+{
+  "info" : {
+    "author" : "xcode",
+    "version" : 1
+  },
+  "textures" : [
+    {
+      "filename" : "Universal.mipmapset",
+      "idiom" : "universal"
+    }
+  ]
+}
diff --git a/Engine/Extras/VirtualProduction/LiveLinkVCAM/vcam/Assets.xcassets/no_signal.textureset/Universal.mipmapset/Contents.json b/Engine/Extras/VirtualProduction/LiveLinkVCAM/vcam/Assets.xcassets/no_signal.textureset/Universal.mipmapset/Contents.json
new file mode 100644
--- /dev/null
+++ b/Engine/Extras/VirtualProduction/LiveLinkVCAM/vcam/Assets.xcassets/no_signal.textureset/Universal.mipmapset/Contents.json
@@ -0,0 +1,12 @@
+{
+  "info" : {
+    "author" : "xcode",
+    "version" : 1
+  },
+  "levels" : [
+    {
+      "filename" : "no-signal.jpg",
+      "mipmap-level" : "base"
+    }
+  ]
+}
diff --git a/Engine/Extras/VirtualProduction/LiveLinkVCAM/vcam/Assets.xcassets/tentacle.imageset/Contents.json b/Engine/Extras/VirtualProduction/LiveLinkVCAM/vcam/Assets.xcassets/tentacle.imageset/Contents.json
new file mode 100644
--- /dev/null
+++ b/Engine/Extras/VirtualProduction/LiveLinkVCAM/vcam/Assets.xcassets/tentacle.imageset/Contents.json
@@ -0,0 +1,23 @@
+{
+  "images" : [
+    {
+      "filename" : "TS_white_20h.png",
+      "idiom" : "universal",
+      "scale" : "1x"
+    },
+    {
+      "filename" : "TS_white_40h.png",
+      "idiom" : "universal",
+      "scale" : "2x"
+    },
+    {
+      "filename" : "TS_white_60h.png",
+      "idiom" : "universal",
+      "scale" : "3x"
+    }
+  ],
+  "info" : {
+    "author" : "xcode",
+    "version" : 1
+  }
+}
diff --git a/Engine/Extras/VirtualProduction/LiveLinkVCAM/vcam/ObjCUtility.h b/Engine/Extras/VirtualProduction/LiveLinkVCAM/vcam/ObjCUtility.h
new file mode 100644
--- /dev/null
+++ b/Engine/Extras/VirtualProduction/LiveLinkVCAM/vcam/ObjCUtility.h
@@ -0,0 +1,25 @@
+//
+//  ObjCUtility.h
+//  Live Link VCAM
+//
+//  Created by Brian Smith on 1/27/20.
+//  Copyright © 2020 Epic Games, Inc. All rights reserved.
+//
+
+#ifndef ObjCUtility_h
+#define ObjCUtility_h
+
+#import <Foundation/Foundation.h>
+
+@class Timecode;
+
+@interface ObjCUtility : NSObject
+
++ (NSDate *) buildDate;
+
+
+
+@end
+
+
+#endif /* ObjCUtility_h */
diff --git a/Engine/Extras/VirtualProduction/LiveLinkVCAM/vcam/ObjCUtility.m b/Engine/Extras/VirtualProduction/LiveLinkVCAM/vcam/ObjCUtility.m
new file mode 100644
--- /dev/null
+++ b/Engine/Extras/VirtualProduction/LiveLinkVCAM/vcam/ObjCUtility.m
@@ -0,0 +1,28 @@
+//
+//  ObjCUtility.m
+//  Live Link VCAM
+//
+//  Created by Brian Smith on 1/27/20.
+//  Copyright © 2020 Epic Games, Inc. All rights reserved.
+//
+
+#import <CoreVideo/CoreVideo.h>
+#import <CoreServices/CoreServices.h>
+
+#import "VCAM-Swift.h"
+#import "ObjCUtility.h"
+
+@implementation ObjCUtility
+
++ (NSDate *) buildDate {
+    
+    NSString *compileDateTime = [NSString stringWithFormat:@"%s %s", __DATE__, __TIME__];
+    
+    NSDateFormatter *df = [NSDateFormatter new];
+    [df setDateFormat:@"MMM d yyyy H:mm:ss"];
+    [df setLocale:[NSLocale localeWithLocaleIdentifier:@"en_US"]];
+
+    return [df dateFromString:compileDateTime];
+}
+
+@end
diff --git a/Engine/Extras/VirtualProduction/LiveLinkVCAM/vcam/PassthroughShaderTypes.h b/Engine/Extras/VirtualProduction/LiveLinkVCAM/vcam/PassthroughShaderTypes.h
new file mode 100644
--- /dev/null
+++ b/Engine/Extras/VirtualProduction/LiveLinkVCAM/vcam/PassthroughShaderTypes.h
@@ -0,0 +1,51 @@
+/*
+See LICENSE folder for this sample’s licensing information.
+
+Abstract:
+Header containing types and enum constants shared between Metal shaders and C/ObjC source
+*/
+
+#ifndef AAPLShaderTypes_h
+#define AAPLShaderTypes_h
+
+#include <simd/simd.h>
+
+// Buffer index values shared between shader and C code to ensure Metal shader buffer inputs
+// match Metal API buffer set calls.
+typedef enum PassthroughVertexInputIndex
+{
+    PassthroughVertexInputIndexVertices     = 0,
+    PassthroughVertexInputIndexUniforms
+} PassthroughVertexInputIndex;
+
+struct PassthroughVertexUniforms {
+    float aspectRatio;
+};
+
+// Texture index values shared between shader and C code to ensure Metal shader buffer inputs match
+//   Metal API texture set calls
+typedef enum BasicTextureInputIndex
+{
+    BasicTextureInputIndexBaseColor = 0,
+    BasicTextureInputIndexUniforms
+} BasicTextureIndex;
+
+struct BasicTextureUniforms {
+    
+    int isARGB;
+    int flipY;
+};
+
+
+
+//  This structure defines the layout of vertices sent to the vertex
+//  shader. This header is shared between the .metal shader and C code, to guarantee that
+//  the layout of the vertex array in the C code matches the layout that the .metal
+//  vertex shader expects.
+typedef struct
+{
+    vector_float2 position;
+    vector_float2 textureCoord;
+} PassthroughVertex;
+
+#endif /* AAPLShaderTypes_h */
diff --git a/Engine/Extras/VirtualProduction/LiveLinkVCAM/vcam/VCAM-Bridging-Header.h b/Engine/Extras/VirtualProduction/LiveLinkVCAM/vcam/VCAM-Bridging-Header.h
new file mode 100644
--- /dev/null
+++ b/Engine/Extras/VirtualProduction/LiveLinkVCAM/vcam/VCAM-Bridging-Header.h
@@ -0,0 +1,16 @@
+//
+//  VCAM-Bridging-Header.h
+//  vcam
+//
+//  Created by Brian Smith on 8/10/20.
+//  Copyright Epic Games. Inc. All rights reserved.
+//
+
+#ifndef VCAM_Bridging_Header_h
+#define VCAM_Bridging_Header_h
+
+#import "PassthroughShaderTypes.h"
+#import "ObjCUtility.h"
+#import <LiveLinkProvider/LiveLink.h>
+
+#endif /* VCAM_Bridging_Header_h */
diff --git a/Engine/Extras/VirtualProduction/LiveLinkVCAM/vcam/eula.html b/Engine/Extras/VirtualProduction/LiveLinkVCAM/vcam/eula.html
new file mode 100644
--- /dev/null
+++ b/Engine/Extras/VirtualProduction/LiveLinkVCAM/vcam/eula.html
@@ -0,0 +1,121 @@
+<style>
+body {
+  font-family: system-ui, -apple-system, BlinkMacSystemFont, "Helvetica Neue",
+      "Lucida Grande", "Segoe UI";
+}
+</style>
+<font size="+1">
+<br>
+<center>Live Link VCAM End User License Agreement</center>
+<br>
+<p>Please read this Agreement carefully.  It is a legal document that explains your rights and obligations related to your use of the Software, including any Services you access through the Software.  By downloading or using the Software, or by otherwise indicating your acceptance of this Agreement, you are agreeing to be bound by the terms of this Agreement.  If you do not or cannot agree to the terms of this Agreement, please do not download or use this Software.
+
+<p>Your use of the Services is also governed by Epic's Terms of Service, which may be found at https://www.epicgames.com/tos and Epic's Privacy Policy, which may be found at https://www.epicgames.com/privacypolicy. By downloading or using the Software, you also agree to Epic's Terms of Service and Privacy Policy.
+
+<p>Certain words or phrases are defined to have certain meanings when used in this Agreement.  Those words and phrases are defined below in Section 14.
+
+<p>If your primary residence is in the United States of America, your agreement is with Epic Games,Inc.  If it is not in the United States of America, your agreement is with Epic Games International S.&agrave; r.l., acting through its Swiss branch.
+
+<p><b>1. License Grant</b>
+
+<p>Epic grants you a personal, non-exclusive, non-transferable, non-sublicensable limited right and license to use and install the Software on an Authorized Device (the <b>"License"</b>).  The rights that Epic grants you under the License are subject to the terms of this Agreement, and you may only make use of the License if you comply with all applicable terms.
+
+<p>The License becomes effective on the date you accept this Agreement. The Software is licensed, not sold, to you under the License. The License does not grant you any title or ownership in the Software.
+
+<p><b>2. License Conditions</b>
+
+<p>You may not do any of the following with respect to the Software or any of its parts:  (a) copy, reproduce, distribute, display, or use it in a way that is not expressly authorized in this Agreement; (b) sell, rent, lease, license, distribute, or otherwise transfer it; (c) reverse engineer, derive source code from, modify, adapt, translate, decompile, or disassemble it or make derivative works based on it; (d) remove, disable, circumvent, or modify any proprietary notice or label or security technology included in it; (e) create, develop, distribute, or use any unauthorized software programs; (f) use it to infringe or violate the rights of any third party, including but not limited to any intellectual property, publicity, or privacy rights; (g) use, export, or re-export it in violation of any applicable law or regulation; (h) use it on a device other than anAuthorized Device; or (i) behave in a manner which is detrimental to the enjoyment of the Software or any Apps by other users as intended by Epic, in Epic's sole judgment, including but not limited to the following – harassment, use of abusive or offensive language, game abandonment, game sabotage, spamming, social engineering, or scamming.
+
+<p><b>3. Updates and Patches</b>
+
+<p>Epic may provide patches, updates, or upgrades to the Software that must be installed in order for you to continue to use the Software or Services.  Epic may update the Software remotely without notifying you, and you hereby consent to Epic applying patches, updates, and upgrades.  Epic may modify, suspend, discontinue, substitute, replace, or limit your access to any aspect of the Software or Services at any time.  You acknowledge that your use of the Software or Services does not confer on you any interest, monetary or otherwise, in any aspect or feature of the Software or any Apps.  Epic does not have any maintenance or support obligations with respect to the Software or Services.
+
+<p><b>4. Feedback</b>
+
+<p>If you provide Epic with any Feedback, you hereby grant Epic a non-exclusive, fully-paid, royalty-free, irrevocable, perpetual, transferable, sublicensable license to reproduce, distribute, modify, prepare derivative works based on, publicly perform, publicly display, make, have made,use, sell, offer to sell, import, and otherwise exploit that Feedback for any purposes, for all current and future methods and forms of exploitation in any country.  If any such rights may not be licensed under applicable law (such as moral and other personal rights), you hereby waive and agree not to assert all such rights.  You understand and agree that Epic is not required to make any use of any Feedback that you provide.  You agree that if Epic makes use of your Feedback, Epic is not required to credit or compensate you for your contribution. You represent and warrant that you have sufficient rights in any Feedback that you provide to Epic to grant Epic and other affected parties the rights described above.  This includes but is not limited to intellectual property rights and other proprietary or personal rights.
+
+<p><b>5. Ownership/Third Party Licenses</b>
+
+<p>Epic and its licensors own all title, ownership rights, and intellectual property rights in the Software and Services.  Epic, Epic Games, Unreal, Unreal Engine, and their respective logos, are trademarks or registered trademarks of Epic and its affiliates in the United States of America and elsewhere.  All rights granted to you under this Agreement are granted by express license only and not by sale.  No license or other rights shall be created hereunder by implication, estoppel, orotherwise.
+
+<p>The Software includes certain components provided by Epic's licensors.  A list of credits and notices for third party components may be found in the Software interface.
+
+<p><b>6. Disclaimers and Limitation of Liability</b>
+
+<p><b>The Software and Services is provided on an "as is" and "as available" basis, "with all faults" and without warranty of any kind.  Epic, its licensors, and its and their affiliates disclaim all warranties, conditions, common law duties, and representations (express, implied, oral, and written) with respect to the Software and Services, including without limitation all express, implied, and statutory warranties and conditions of any kind, such astitle, non-interference with your enjoyment, authority, non-infringement, merchantability, fitness or suitability for any purpose (whether or not Epic knows or has reason to know of any such purpose), system integration, accuracy or completeness, results, reasonable care, workmanlike effort, lack of negligence, and lack of viruses, whether alleged to arise under law, by reason of custom or usage in the trade, or by course of dealing.  Without limiting the generality of the foregoing, Epic, its licensors, and its and their affiliates make no warranty that (1) the Software or Services will operate properly, (2) that the Software or Services will meet your requirements, (3) that the operation of the Software or Services willbe uninterrupted, bug free, or error free in any or all circumstances, or (4) that any defects in the Software or Services can or will be corrected.  Any warranty against infringement that may be provided in Section 2-312 of the Uniform Commercial Code or in any other comparable statute is expressly disclaimed.  Epic, its licensors, and its and their affiliates donot guarantee continuous, error-free, virus-free, or secure operation of or access to the Software or Services.  This paragraph will apply to the maximum extent permitted by applicable law.</b>
+
+<p><b>To the maximum extent permitted by applicable law, neither Epic, nor its licensors, nor its or their affiliates, nor any of Epic's service providers (collectively, the "Epic Parties"), shall be liable in any way for any loss of profits or any indirect, incidental, consequential, special, punitive, or exemplary damages, arising out of or in connection with this Agreement or the Software or Services, or the delay or inability to use or lack of functionality of the Software or Services, even in the event of an Epic Party's fault, tort (including negligence), strict liability, indemnity, product liability, breach of contract, breach of warranty, or otherwise and even if an Epic Party has been advised of the possibility of such damages.  Further, to the maximum extent permitted by applicable law, the aggregate liability of the Epic Parties arising out of or in connection with this Agreement or the Software  or Services will not exceed the greater of (a) the total amounts you have paid (if any) to Epic for the Software or within the Software during the twelve (12) months immediately preceding the events giving rise to such liability and (b) U.S. $5.00. These limitations and exclusions regarding damages apply even if any remedy fails to provide adequate compensation.</b>
+
+<p><b>Notwithstanding the foregoing, some countries, states, provinces or other jurisdictions do not allow the exclusion of certain warranties or the limitation of liability as stated above, so the above terms may not apply to you. Instead, in such jurisdictions, the foregoing exclusions and limitations shall apply only to the extent permitted by the laws of such jurisdictions. Also, you may have additional legal rights in your jurisdiction, and nothing in this Agreement will prejudice the statutory rights that you may have as a consumer of the Software or Services.</b>
+
+<p><b>7. Indemnity</b>
+
+<p>You agree to indemnify, pay the defense costs of, and hold Epic, its licensors, its and their affiliates, and its and their employees, officers, directors, agents, contractors, and other representatives harmless from all claims, demands, actions, losses, liabilities, and expenses (including attorneys' fees, costs, and expert witnesses' fees) that arise from or in connection with (a) any claim that, if true, would constitute a breach by you of this Agreement or negligence by you, or (b) any act or omission by you in using the Software or Services.You agree to reimburseEpic on demand for any defense costs incurred by Epic and any payments made or loss suffered by Epic, whether in a court judgment or settlement, based on any matter covered by this Section 7.
+
+<p>If you are prohibited by law from entering into the indemnification obligation above, then you assume, to the extent permitted by law, all liability for all claims, demands, actions, losses, liabilities, and expenses (including attorneys' fees, costs and expert witnesses' fees) that are the stated subject matter of the indemnification obligation above.
+
+<p><b>8. Termination</b>
+
+<p>Without limiting any other rights of Epic, this Agreement will terminate automatically without notice if you fail to comply with any of its terms and conditions.  You may also terminate this Agreement by deleting the Software.  Upon any termination, the License will automatically terminate, you may no longer exercise any of the rights granted to you by the License, and you must destroy all copies of the Software in your possession.
+
+<p>Except to the extent required by law, all payments and fees are non-refundable under all circumstances, regardless of whether or not this Agreement has been terminated.
+
+<p>Sections 2, 6-11, and 13-15 will survive any termination of this Agreement.
+
+<p><b>9. Governing Law and Jurisdiction</b>
+
+<p>You agree that this Agreement will be deemed to have been made and executed in the State of North Carolina, U.S.A., and any dispute will be resolved in accordance with the laws of North Carolina, excluding that body of law related to choice of laws, and of the United States of America.  Any action or proceeding brought to enforce the terms of this Agreement or to adjudicate any dispute must be brought in the Superior Court of Wake County, State of North Carolina or the United States District Court for the Eastern District of North Carolina.  You agreeto the exclusive jurisdiction and venue of these courts.  You waive any claim of inconvenient forum and any right to a jury trial.  The Convention on Contracts for the International Sale of Goods will not apply.Any law or regulation which provides that the language of a contract shallbe construed against the drafter will not apply to this Agreement.
+
+<p><b>10. Class Action Waiver</b>
+
+<p><b>You agree not to bring or participate in a class or representative action, private attorney general action, or collective arbitration related to the Software or Services or this Agreement.  You also agree not to seek to combine any action or arbitration related to the Software or Services or this Agreement with any other action or arbitration without the consent of all parties to this Agreement and all other actions or arbitrations.</b>
+
+<p><b>11. U.S. Government Matters</b>
+
+<p>The Software is a "Commercial Item" (as defined at 48 C.F.R. &sect;2.101), consisting of "Commercial Computer Software" and "Commercial Computer Software Documentation" (as used in 48 C.F.R. &sect;12.212 or 48 C.F.R. &sect;227.7202, as applicable).  The Software is being licensed to U.S. Government end users only as Commercial Items and with only those rights as are granted to other licensees under this Agreement.
+
+<p>You represent and warrant to Epic that you are not located in a country that is subject to a U.S. Government embargo or that has been designated by the U.S. Government as a "terrorist supporting" country, and that you are not listed on any U.S. Government list of prohibited or restricted parties.
+
+<p><b>12. Amendments of this Agreement</b>
+
+<p><b>Epic may issue an amended Agreement, Terms of Service, or Privacy Policy at any time in its discretion by posting the amended Agreement, Terms of Service, or Privacy Policy on its website or by providing you with digital access to amended versions of any of these documents when you next access the Software.  If any amendment to this Agreement, the Terms of Service, or the Privacy Policy is not acceptable to you, you may terminate this Agreement and must stop using the Software.  Your continued use of the Software will demonstrate your acceptance of the amended Agreement, Terms of Service, and Privacy Policy.</b>
+
+<p><b>13. No Assignment</b>
+
+<p>You may not, without the prior written consent of Epic, assign, transfer, charge, or sub-contract all or any of your rights or obligations under this Agreement, and any attempt without that consent will be null and void.  If restrictions on transfer of the Software in this Agreement are not enforceable under the law of your country, then this Agreement will be binding on any recipient of the Software.  Epic may at any time assign, transfer, charge, or sub-contract all or any of its rights or obligations under this Agreement.
+
+<p><b>14. Definitions</b>
+
+<p>As used in this Agreement, the following capitalized words have the following meanings:
+
+<p><b>"Authorized Device"</b> means a device on which the Software is designed and designated by Epic to be operated. An Authorized Device does not include a virtual machine or other emulation of such a mobile device.
+
+<p><b>"Epic"</b> means, depending on the location of your primary residence:
+
+<ol type="a">
+
+<li>Epic Games, Inc., a Maryland Corporation having its principal business offices at Box254, 2474 Walnut Street, Cary, North Carolina, 27518, U.S.A.; or
+
+<li>Epic Games International S.&agrave; r.l., a Luxembourg Soci&eacute;t&eacute; &agrave; Responsibilit&eacute; Limit&eacute;e, located at Atrium Business Park, 33 rue du Puits Romain, L8070 Bertrange, Grand-Duchy of Luxembourg, acting through its Swiss branch, having its principal business offices at Platz 3, 6039 Root, Switzerland.
+</ol>
+<br>
+
+<p><b>"Feedback"</b> means any feedback or suggestions that you provide to Epic regarding the Software, Services, or other Epic products and services.
+
+<p><b>"Services"</b> means any services made available to you through the Software, including services tomanage your Epic account, connect to third party services.
+
+<p><b>"Software"</b> means Epic's proprietary software known as Unreal Face.
+
+<p><b>15. Miscellaneous</b>
+
+<p>This Agreement and any document or information referred to in this Agreement constitute the entire agreement between you and Epic relating to the subject matter covered by this Agreement. All other communications, proposals, and representations with respect to the subject matter covered by this Agreement are excluded.
+
+<p>The original of this Agreement is in English; any translations are provided for reference purposes only.  You waive any right you may have under the law of your country to have this Agreement written or construed in the language of any other country.
+
+<p>This Agreement describes certain legal rights. You may have other rights under the laws of your jurisdiction. This Agreement does not change your rights under the laws of your jurisdiction if the laws of your jurisdiction do not permit it to do so.  Limitations and exclusions of warranties and remedies in this Agreement may not apply to you because your jurisdiction may not allow them in your particular circumstance.  In the event that certain provisions of this Agreement are held by a court or tribunal of competent jurisdiction to be unenforceable, those provisions shall be enforced only to the furthest extent possible under applicable law and the remaining terms of this Agreement will remain in full force and effect.
+
+<p>Any act by Epic to exercise, or failure or delay in exercise of, any of its rights under this Agreement, at law or in equity will not be deemed a waiver of those or any other rights or remedies available in contract, at law or in equity.
+
+<p>You agree that this Agreement does not confer any rights or remedies on any person other than the parties to this Agreement, except as expressly stated.Epic's obligations are subject to existing laws and legal process, and Epic may comply with law enforcement or regulatory requests or requirements despite any contrary term in this Agreement.
+
+</font>
