Fixed horde issue 120870 : missing copyright boilerplates.

diff --git a/Engine/Extras/VirtualProduction/LiveLinkVCAM/LiveLinkProvider.framework/Headers/LiveLink.h b/Engine/Extras/VirtualProduction/LiveLinkVCAM/LiveLinkProvider.framework/Headers/LiveLink.h
--- a/Engine/Extras/VirtualProduction/LiveLinkVCAM/LiveLinkProvider.framework/Headers/LiveLink.h
+++ b/Engine/Extras/VirtualProduction/LiveLinkVCAM/LiveLinkProvider.framework/Headers/LiveLink.h
@@ -1,3 +1,4 @@
+//  Copyright Epic Games. Inc. All rights reserved.
 
 #import <Foundation/Foundation.h>
 
diff --git a/Engine/Extras/VirtualProduction/LiveLinkVCAM/LiveLinkProvider.framework/Headers/LiveLinkProvider.h b/Engine/Extras/VirtualProduction/LiveLinkVCAM/LiveLinkProvider.framework/Headers/LiveLinkProvider.h
--- a/Engine/Extras/VirtualProduction/LiveLinkVCAM/LiveLinkProvider.framework/Headers/LiveLinkProvider.h
+++ b/Engine/Extras/VirtualProduction/LiveLinkVCAM/LiveLinkProvider.framework/Headers/LiveLinkProvider.h
@@ -1,3 +1,4 @@
+//  Copyright Epic Games. Inc. All rights reserved.
 
 #import <Foundation/Foundation.h>
 #import <simd/simd.h>
diff --git a/Engine/Extras/VirtualProduction/LiveLinkVCAM/vcam/ObjCUtility.h b/Engine/Extras/VirtualProduction/LiveLinkVCAM/vcam/ObjCUtility.h
--- a/Engine/Extras/VirtualProduction/LiveLinkVCAM/vcam/ObjCUtility.h
+++ b/Engine/Extras/VirtualProduction/LiveLinkVCAM/vcam/ObjCUtility.h
@@ -1,10 +1,4 @@
-//
-//  ObjCUtility.h
-//  Live Link VCAM
-//
-//  Created by Brian Smith on 1/27/20.
-//  Copyright © 2020 Epic Games, Inc. All rights reserved.
-//
+//  Copyright Epic Games. Inc. All rights reserved.
 
 #ifndef ObjCUtility_h
 #define ObjCUtility_h
diff --git a/Engine/Extras/VirtualProduction/LiveLinkVCAM/vcam/ObjCUtility.m b/Engine/Extras/VirtualProduction/LiveLinkVCAM/vcam/ObjCUtility.m
--- a/Engine/Extras/VirtualProduction/LiveLinkVCAM/vcam/ObjCUtility.m
+++ b/Engine/Extras/VirtualProduction/LiveLinkVCAM/vcam/ObjCUtility.m
@@ -1,10 +1,4 @@
-//
-//  ObjCUtility.m
-//  Live Link VCAM
-//
-//  Created by Brian Smith on 1/27/20.
-//  Copyright © 2020 Epic Games, Inc. All rights reserved.
-//
+//  Copyright Epic Games. Inc. All rights reserved.
 
 #import <CoreVideo/CoreVideo.h>
 #import <CoreServices/CoreServices.h>
diff --git a/Engine/Extras/VirtualProduction/LiveLinkVCAM/vcam/VCAM-Bridging-Header.h b/Engine/Extras/VirtualProduction/LiveLinkVCAM/vcam/VCAM-Bridging-Header.h
--- a/Engine/Extras/VirtualProduction/LiveLinkVCAM/vcam/VCAM-Bridging-Header.h
+++ b/Engine/Extras/VirtualProduction/LiveLinkVCAM/vcam/VCAM-Bridging-Header.h
@@ -1,10 +1,5 @@
-//
-//  VCAM-Bridging-Header.h
-//  vcam
-//
-//  Created by Brian Smith on 8/10/20.
 //  Copyright Epic Games. Inc. All rights reserved.
-//
+
 
 #ifndef VCAM_Bridging_Header_h
 #define VCAM_Bridging_Header_h
