Fixed horde issue 120871 : missing copyright boilerplates.

diff --git a/Engine/Extras/VirtualProduction/LiveLinkVCAM/LiveLinkProvider.framework/Headers/LiveLink.h b/Engine/Extras/VirtualProduction/LiveLinkVCAM/LiveLinkProvider.framework/Headers/LiveLink.h
--- a/Engine/Extras/VirtualProduction/LiveLinkVCAM/LiveLinkProvider.framework/Headers/LiveLink.h
+++ b/Engine/Extras/VirtualProduction/LiveLinkVCAM/LiveLinkProvider.framework/Headers/LiveLink.h
@@ -1,4 +1,4 @@
-//  Copyright Epic Games. Inc. All rights reserved.
+//  Copyright Epic Games, Inc. All Rights Reserved.
 
 #import <Foundation/Foundation.h>
 
diff --git a/Engine/Extras/VirtualProduction/LiveLinkVCAM/LiveLinkProvider.framework/Headers/LiveLinkProvider.h b/Engine/Extras/VirtualProduction/LiveLinkVCAM/LiveLinkProvider.framework/Headers/LiveLinkProvider.h
--- a/Engine/Extras/VirtualProduction/LiveLinkVCAM/LiveLinkProvider.framework/Headers/LiveLinkProvider.h
+++ b/Engine/Extras/VirtualProduction/LiveLinkVCAM/LiveLinkProvider.framework/Headers/LiveLinkProvider.h
@@ -1,4 +1,4 @@
-//  Copyright Epic Games. Inc. All rights reserved.
+//  Copyright Epic Games, Inc. All Rights Reserved.
 
 #import <Foundation/Foundation.h>
 #import <simd/simd.h>
diff --git a/Engine/Extras/VirtualProduction/LiveLinkVCAM/vcam/ObjCUtility.h b/Engine/Extras/VirtualProduction/LiveLinkVCAM/vcam/ObjCUtility.h
--- a/Engine/Extras/VirtualProduction/LiveLinkVCAM/vcam/ObjCUtility.h
+++ b/Engine/Extras/VirtualProduction/LiveLinkVCAM/vcam/ObjCUtility.h
@@ -1,4 +1,4 @@
-//  Copyright Epic Games. Inc. All rights reserved.
+//  Copyright Epic Games, Inc. All Rights Reserved.
 
 #ifndef ObjCUtility_h
 #define ObjCUtility_h
diff --git a/Engine/Extras/VirtualProduction/LiveLinkVCAM/vcam/ObjCUtility.m b/Engine/Extras/VirtualProduction/LiveLinkVCAM/vcam/ObjCUtility.m
--- a/Engine/Extras/VirtualProduction/LiveLinkVCAM/vcam/ObjCUtility.m
+++ b/Engine/Extras/VirtualProduction/LiveLinkVCAM/vcam/ObjCUtility.m
@@ -1,4 +1,4 @@
-//  Copyright Epic Games. Inc. All rights reserved.
+//  Copyright Epic Games, Inc. All Rights Reserved.
 
 #import <CoreVideo/CoreVideo.h>
 #import <CoreServices/CoreServices.h>
diff --git a/Engine/Extras/VirtualProduction/LiveLinkVCAM/vcam/PassthroughShaderTypes.h b/Engine/Extras/VirtualProduction/LiveLinkVCAM/vcam/PassthroughShaderTypes.h
--- a/Engine/Extras/VirtualProduction/LiveLinkVCAM/vcam/PassthroughShaderTypes.h
+++ b/Engine/Extras/VirtualProduction/LiveLinkVCAM/vcam/PassthroughShaderTypes.h
@@ -1,9 +1,4 @@
-/*
-See LICENSE folder for this sample’s licensing information.
-
-Abstract:
-Header containing types and enum constants shared between Metal shaders and C/ObjC source
-*/
+//  Copyright Epic Games, Inc. All Rights Reserved.
 
 #ifndef AAPLShaderTypes_h
 #define AAPLShaderTypes_h
diff --git a/Engine/Extras/VirtualProduction/LiveLinkVCAM/vcam/VCAM-Bridging-Header.h b/Engine/Extras/VirtualProduction/LiveLinkVCAM/vcam/VCAM-Bridging-Header.h
--- a/Engine/Extras/VirtualProduction/LiveLinkVCAM/vcam/VCAM-Bridging-Header.h
+++ b/Engine/Extras/VirtualProduction/LiveLinkVCAM/vcam/VCAM-Bridging-Header.h
@@ -1,4 +1,4 @@
-//  Copyright Epic Games. Inc. All rights reserved.
+//  Copyright Epic Games, Inc. All Rights Reserved.
 
 
 #ifndef VCAM_Bridging_Header_h
