Live Link VCAM (iOS) v1.0.2 build 3 : Added UE4 vs UE5 LiveLink support. The Large World Coordinates change (among others) prevents compatibility between versions when sending transforms. There are now two LiveLinkProvider frameworks.  This is configurable in the app but requires an app restart to load the proper libray.

diff --git a/Engine/Extras/VirtualProduction/LiveLinkVCAM/LiveLinkProvider.framework/Headers/LiveLink.h b/Engine/Extras/VirtualProduction/LiveLinkVCAM/LiveLinkProvider.framework/Headers/LiveLink.h
deleted file mode 100644
diff --git a/Engine/Extras/VirtualProduction/LiveLinkVCAM/LiveLinkProvider.framework/Headers/LiveLinkProvider.h b/Engine/Extras/VirtualProduction/LiveLinkVCAM/LiveLinkProvider.framework/Headers/LiveLinkProvider.h
deleted file mode 100644
diff --git a/Engine/Extras/VirtualProduction/LiveLinkVCAM/UE4LiveLinkProvider.framework/Headers/UE4LiveLink.h b/Engine/Extras/VirtualProduction/LiveLinkVCAM/UE4LiveLinkProvider.framework/Headers/UE4LiveLink.h
new file mode 100644
--- /dev/null
+++ b/Engine/Extras/VirtualProduction/LiveLinkVCAM/UE4LiveLinkProvider.framework/Headers/UE4LiveLink.h
@@ -0,0 +1,23 @@
+//  Copyright Epic Games, Inc. All Rights Reserved.
+
+#import <Foundation/Foundation.h>
+
+#import "UE4LiveLinkProvider.h"
+
+__attribute__((visibility("default")))
+@protocol UE4LiveLinkLogDelegate
+
+- (void) ue4LogMessage:(NSString *)message;
+
+@end
+
+ 
+__attribute__((visibility("default")))
+@interface UE4LiveLink : NSObject
+
++ (void) initialize:(id<UE4LiveLinkLogDelegate>)logDelegate;
++ (void) restart;
++ (void) shutdown;
++ (id<UE4LiveLinkProvider> ) createProvider:(NSString *)name;
+
+@end
diff --git a/Engine/Extras/VirtualProduction/LiveLinkVCAM/UE4LiveLinkProvider.framework/Headers/UE4LiveLinkProvider.h b/Engine/Extras/VirtualProduction/LiveLinkVCAM/UE4LiveLinkProvider.framework/Headers/UE4LiveLinkProvider.h
new file mode 100644
--- /dev/null
+++ b/Engine/Extras/VirtualProduction/LiveLinkVCAM/UE4LiveLinkProvider.framework/Headers/UE4LiveLinkProvider.h
@@ -0,0 +1,15 @@
+//  Copyright Epic Games, Inc. All Rights Reserved.
+
+#import <Foundation/Foundation.h>
+#import <simd/simd.h>
+
+@protocol UE4LiveLinkProvider <NSObject>
+
+- (void)addCameraSubject:(NSString *)subjectName;
+- (void)removeCameraSubject:(NSString *)subjectName;
+
+- (void)updateSubject:(NSString *)subjectName withTransform:(simd_float4x4)xform;
+- (void)updateSubject:(NSString *)subjectName withTransform:(simd_float4x4)xform atTime:(double)time;
+
+@end
+
diff --git a/Engine/Extras/VirtualProduction/LiveLinkVCAM/UE5LiveLinkProvider.framework/Headers/UE5LiveLink.h b/Engine/Extras/VirtualProduction/LiveLinkVCAM/UE5LiveLinkProvider.framework/Headers/UE5LiveLink.h
new file mode 100644
--- /dev/null
+++ b/Engine/Extras/VirtualProduction/LiveLinkVCAM/UE5LiveLinkProvider.framework/Headers/UE5LiveLink.h
@@ -0,0 +1,23 @@
+//  Copyright Epic Games, Inc. All Rights Reserved.
+
+#import <Foundation/Foundation.h>
+
+#import "UE5LiveLinkProvider.h"
+
+__attribute__((visibility("default")))
+@protocol UE5LiveLinkLogDelegate
+
+- (void) ue5LogMessage:(NSString *)message;
+
+@end
+
+ 
+__attribute__((visibility("default")))
+@interface UE5LiveLink : NSObject
+
++ (void) initialize:(id<UE5LiveLinkLogDelegate>)logDelegate;
++ (void) restart;
++ (void) shutdown;
++ (id<UE5LiveLinkProvider> ) createProvider:(NSString *)name;
+
+@end
diff --git a/Engine/Extras/VirtualProduction/LiveLinkVCAM/UE5LiveLinkProvider.framework/Headers/UE5LiveLinkProvider.h b/Engine/Extras/VirtualProduction/LiveLinkVCAM/UE5LiveLinkProvider.framework/Headers/UE5LiveLinkProvider.h
new file mode 100644
--- /dev/null
+++ b/Engine/Extras/VirtualProduction/LiveLinkVCAM/UE5LiveLinkProvider.framework/Headers/UE5LiveLinkProvider.h
@@ -0,0 +1,15 @@
+//  Copyright Epic Games, Inc. All Rights Reserved.
+
+#import <Foundation/Foundation.h>
+#import <simd/simd.h>
+
+@protocol UE5LiveLinkProvider <NSObject>
+
+- (void)addCameraSubject:(NSString *)subjectName;
+- (void)removeCameraSubject:(NSString *)subjectName;
+
+- (void)updateSubject:(NSString *)subjectName withTransform:(simd_float4x4)xform;
+- (void)updateSubject:(NSString *)subjectName withTransform:(simd_float4x4)xform atTime:(double)time;
+
+@end
+
diff --git a/Engine/Extras/VirtualProduction/LiveLinkVCAM/vcam/AppDelegate.swift b/Engine/Extras/VirtualProduction/LiveLinkVCAM/vcam/AppDelegate.swift
--- a/Engine/Extras/VirtualProduction/LiveLinkVCAM/vcam/AppDelegate.swift
+++ b/Engine/Extras/VirtualProduction/LiveLinkVCAM/vcam/AppDelegate.swift
@@ -11,8 +11,6 @@ import UIKit
 @UIApplicationMain
 class AppDelegate: UIResponder, UIApplicationDelegate {
 
-
-
     func application(_ application: UIApplication, didFinishLaunchingWithOptions launchOptions: [UIApplication.LaunchOptionsKey: Any]?) -> Bool {
         // Override point for customization after application launch.
         AppSettings.registerDefaults()
@@ -32,7 +30,7 @@ class AppDelegate: UIResponder, UIApplicationDelegate {
     }
     
     func applicationWillTerminate(_ application: UIApplication) {
-        LiveLink.shutdown()
+        try? LiveLink.shutdown()
     }
 
     // MARK: UISceneSession Lifecycle
@@ -48,7 +46,5 @@ class AppDelegate: UIResponder, UIApplicationDelegate {
         // If any sessions were discarded while the application was not running, this will be called shortly after application:didFinishLaunchingWithOptions.
         // Use this method to release any resources that were specific to the discarded scenes, as they will not return.
     }
-
-
 }
 
diff --git a/Engine/Extras/VirtualProduction/LiveLinkVCAM/vcam/AppSettings.swift b/Engine/Extras/VirtualProduction/LiveLinkVCAM/vcam/AppSettings.swift
--- a/Engine/Extras/VirtualProduction/LiveLinkVCAM/vcam/AppSettings.swift
+++ b/Engine/Extras/VirtualProduction/LiveLinkVCAM/vcam/AppSettings.swift
@@ -21,11 +21,12 @@ class AppSettings : NSObject {
     class func defaultLiveLinkSubjectName() -> String {
         UIDevice.current.name.toUnrealCompatibleString()
     }
-    
+
     private class Keys {
         static let didAcceptEULA = "didAcceptEULA"
         static let lastConnectionAddress = "lastConnectionAddress"
         
+        static let engineVersion = "engineVersion"
         static let liveLinkSubjectName = "liveLinkSubjectName"
 
         static let timecodeSource = "timecodeSource"
@@ -43,6 +44,7 @@ class AppSettings : NSObject {
             Keys.didAcceptEULA : false,
             Keys.lastConnectionAddress : "",
 
+            Keys.engineVersion : "",
             Keys.liveLinkSubjectName : defaultLiveLinkSubjectName(),
             
             Keys.timecodeSource : 1,
@@ -86,7 +88,13 @@ class AppSettings : NSObject {
             UserDefaults.standard.set(lastConnectionAddress, forKey: Keys.lastConnectionAddress)
         }
     }
-    
+
+    @objc dynamic var engineVersion : String = UserDefaults.standard.string(forKey: Keys.engineVersion) ?? "" {
+        didSet {
+            UserDefaults.standard.set(engineVersion, forKey: Keys.engineVersion)
+        }
+    }
+
     @objc dynamic var liveLinkSubjectName : String = UserDefaults.standard.string(forKey: Keys.liveLinkSubjectName) ?? defaultLiveLinkSubjectName() {
         didSet {
             UserDefaults.standard.set(liveLinkSubjectName, forKey: Keys.liveLinkSubjectName)
diff --git a/Engine/Extras/VirtualProduction/LiveLinkVCAM/vcam/LiveLink.swift b/Engine/Extras/VirtualProduction/LiveLinkVCAM/vcam/LiveLink.swift
new file mode 100644
--- /dev/null
+++ b/Engine/Extras/VirtualProduction/LiveLinkVCAM/vcam/LiveLink.swift
@@ -0,0 +1,144 @@
+//
+//  LiveLink.swift
+//  vcam
+//
+//  Created by Brian Smith on 1/13/22.
+//  Copyright Epic Games, Inc. All Rights Reserved.
+//
+
+import Foundation
+
+enum EngineVersion {
+    case unknown
+    case ue4
+    case ue5
+}
+
+enum LiveLinkError : Error {
+    
+    case unknownEngineVersion
+    case alreadyInitialized
+}
+
+
+class LiveLink {
+    
+    private static var currentEngineVersion = EngineVersion.unknown
+    private var provider : AnyObject?
+
+    static var engineVersion : EngineVersion {
+        get {
+            switch AppSettings.shared.engineVersion {
+            case "UE4": return .ue4
+            case "UE5": return .ue5
+            default: return .unknown
+            }
+        }
+        set {
+            switch newValue {
+            case .ue4:
+                AppSettings.shared.engineVersion = "UE4"
+            case .ue5:
+                AppSettings.shared.engineVersion = "UE5"
+            case .unknown:
+                AppSettings.shared.engineVersion = ""
+            }
+        }
+    }
+    
+    static var initialized : Bool {
+        get {
+            currentEngineVersion != .unknown
+        }
+    }
+    
+    static var requiresRestart : Bool {
+        get {
+            currentEngineVersion != engineVersion
+        }
+    }
+    
+    class func initialize(_ delegate : AnyObject & UE5LiveLinkLogDelegate & UE4LiveLinkLogDelegate ) throws {
+        
+        guard currentEngineVersion == .unknown else { throw LiveLinkError.alreadyInitialized }
+        
+        switch engineVersion {
+        case .ue4:
+            UE4LiveLink.initialize(delegate)
+        case .ue5:
+            UE5LiveLink.initialize(delegate)
+        default:
+            throw LiveLinkError.unknownEngineVersion
+        }
+        
+        LiveLink.currentEngineVersion = self.engineVersion
+    }
+    
+    class func restart() throws {
+        switch currentEngineVersion {
+        case .ue4:
+            UE4LiveLink.restart()
+        case .ue5:
+            UE5LiveLink.restart()
+        default:
+            throw LiveLinkError.unknownEngineVersion
+        }
+    }
+
+    class func shutdown() throws {
+        switch currentEngineVersion {
+        case .ue4:
+            UE4LiveLink.shutdown()
+        case .ue5:
+            UE5LiveLink.shutdown()
+        default:
+            throw LiveLinkError.unknownEngineVersion
+        }
+    }
+    
+    init(_ providerName : String) throws {
+        switch LiveLink.currentEngineVersion {
+        case .ue4:
+            provider = UE4LiveLink.createProvider(providerName)
+        case .ue5:
+            provider = UE5LiveLink.createProvider(providerName)
+        default:
+            throw LiveLinkError.unknownEngineVersion
+        }
+    }
+    
+    deinit {
+        provider = nil
+    }
+    
+    func addCameraSubject(_ subjectName : String){
+        if let p = provider as? UE4LiveLinkProvider {
+            p.addCameraSubject(subjectName)
+        }
+
+        if let p = provider as? UE5LiveLinkProvider {
+            p.addCameraSubject(subjectName)
+        }
+    }
+    
+    func removeCameraSubject(_ subjectName : String){
+        if let p = provider as? UE4LiveLinkProvider {
+            p.removeCameraSubject(subjectName)
+        }
+
+        if let p = provider as? UE5LiveLinkProvider {
+            p.removeCameraSubject(subjectName)
+        }
+    }
+
+    func updateSubject( _ subjectName : String, transform xform : simd_float4x4, time : CFTimeInterval) {
+        if let p = provider as? UE4LiveLinkProvider {
+            p.updateSubject(AppSettings.shared.liveLinkSubjectName, withTransform: xform, atTime: time)
+        }
+        if let p = provider as? UE5LiveLinkProvider {
+            p.updateSubject(AppSettings.shared.liveLinkSubjectName, withTransform: xform, atTime: time)
+        }
+    }
+
+
+}
diff --git a/Engine/Extras/VirtualProduction/LiveLinkVCAM/vcam/MultipleChoiceViewController.swift b/Engine/Extras/VirtualProduction/LiveLinkVCAM/vcam/MultipleChoiceViewController.swift
new file mode 100644
--- /dev/null
+++ b/Engine/Extras/VirtualProduction/LiveLinkVCAM/vcam/MultipleChoiceViewController.swift
@@ -0,0 +1,61 @@
+//
+//  MultipleChoiceViewController.swift
+//  Live Link Face
+//
+//  Created by Brian Smith on 1/16/20.
+//  Copyright Epic Games, Inc. All Rights Reserved.
+//
+
+import UIKit
+
+class MultipleChoiceViewController : UITableViewController {
+    
+    var items : [String]!
+    var selectedIndex : Int!
+    
+    var footerString : String?
+
+    var completion : ((Int) -> Void)? = nil
+    var selectionChanged : ((Int) -> Void)? = nil
+
+    override func viewWillDisappear(_ animated: Bool) {
+
+        if let callback = completion {
+            callback(selectedIndex)
+        }
+        
+        super.viewWillDisappear(animated)
+    }
+    
+    override func tableView(_ tableView: UITableView, numberOfRowsInSection section: Int) -> Int {
+        return items.count
+    }
+    
+    override func tableView(_ tableView: UITableView, cellForRowAt indexPath: IndexPath) -> UITableViewCell {
+        
+        var cell : UITableViewCell!
+        
+        cell = tableView.dequeueReusableCell(withIdentifier: "choiceCell")
+        if cell == nil {
+            cell = UITableViewCell(style: .default, reuseIdentifier: "choiceCell")
+        }
+        
+        cell.textLabel?.text = self.items[indexPath.row]
+        cell.accessoryType = (indexPath.row == selectedIndex) ? .checkmark : .none
+
+        return cell
+        
+    }
+
+    override func tableView(_ tableView: UITableView, didSelectRowAt indexPath: IndexPath) {
+        self.selectedIndex = indexPath.row
+        self.tableView.reloadData()
+        if let callback = selectionChanged {
+            callback(indexPath.row)
+        }
+    }
+    
+    override func tableView(_ tableView: UITableView, titleForFooterInSection section: Int) -> String? {
+        return footerString
+    }
+}
diff --git a/Engine/Extras/VirtualProduction/LiveLinkVCAM/vcam/SettingsViewController.swift b/Engine/Extras/VirtualProduction/LiveLinkVCAM/vcam/SettingsViewController.swift
--- a/Engine/Extras/VirtualProduction/LiveLinkVCAM/vcam/SettingsViewController.swift
+++ b/Engine/Extras/VirtualProduction/LiveLinkVCAM/vcam/SettingsViewController.swift
@@ -39,6 +39,9 @@ class SettingsViewController : UITableViewController {
         case "timecode":
             cell.detailTextLabel?.text = Timecode.sourceToString(appSettings.timecodeSourceEnum())
 
+        case "engineVersion":
+            cell.detailTextLabel?.text = (LiveLink.engineVersion == .ue4 ? "UE4" : "UE5") + (LiveLink.requiresRestart ? " (restart required)" : "")
+
         default:
             break
         }
@@ -65,6 +68,44 @@ class SettingsViewController : UITableViewController {
                         self.tableView.reloadData()
                     }
                 }
+            }
+        } else if let vc = segue.destination as? MultipleChoiceViewController {
+            
+            if segue.identifier == "engineVersion" {
+                vc.navigationItem.title = NSLocalizedString("settings-title-engine-version", value:"Engine Version", comment: "Title of an settings screen.")
+                vc.items = [ "Unreal Engine 4", "Unreal Engine 5" ]
+                vc.selectedIndex = (LiveLink.engineVersion == .ue4) ? 0 : 1
+                vc.footerString = NSLocalizedString("settings-footer-engine-version", value: "The Live Link protocol was updated in Unreal Engine 5.0. Live Link VCAM supports both protocols, but this must be specified to match the version of Unreal Engine in use for the project in order to receive tracking data. After changing this setting the app must be terminated & restarted for the setting to take effect.", comment: "A note about how the LiveLink protocol has changed between UE4 & UE5.")
+                
+                vc.selectionChanged = { (index) in
+
+                    switch index {
+                    case 0:
+                        LiveLink.engineVersion = .ue4
+                    case 1:
+                        LiveLink.engineVersion = .ue5
+                    default:
+                        break
+                    }
+
+                    // if the engine version has changed, prompt the user to restart
+                    if (LiveLink.requiresRestart) {
+                        let errorAlert = UIAlertController(title: NSLocalizedString("alert-title-restart", value: "Restart required", comment: "Title indicating the user must restart the application."), message: NSLocalizedString("alert-message-restart", value: "For this setting to take effect, the application must be terminated and restarted.", comment: "Message indicating the user must restart the application."), preferredStyle: .alert)
+                        errorAlert.addAction(UIAlertAction(title: Localized.buttonOK(), style: .default))
+                        errorAlert.addAction(UIAlertAction(title: NSLocalizedString("alert-button-learnmore", value: "Learn more...", comment: "Button to get more information."), style: .default, handler: { action in
+                            if UIDevice.current.userInterfaceIdiom == .phone {
+                                UIApplication.shared.open(URL(string: "https://support.apple.com/HT201330")!, options: [:])
+                            } else {
+                                UIApplication.shared.open(URL(string: "https://support.apple.com/HT212063")!, options: [:])
+                            }
+                        }))
+                        vc.present(errorAlert, animated: true)
+
+                    }
+                    
+                }
+                
+                
             }
         }
     }
diff --git a/Engine/Extras/VirtualProduction/LiveLinkVCAM/vcam/StartViewController.swift b/Engine/Extras/VirtualProduction/LiveLinkVCAM/vcam/StartViewController.swift
--- a/Engine/Extras/VirtualProduction/LiveLinkVCAM/vcam/StartViewController.swift
+++ b/Engine/Extras/VirtualProduction/LiveLinkVCAM/vcam/StartViewController.swift
@@ -18,6 +18,8 @@ class StartViewController : BaseViewController {
     @IBOutlet weak var headerView : HeaderView!
     @IBOutlet weak var versionLabel : UILabel!
 
+    @IBOutlet weak var restartView : UIView!
+
     @IBOutlet weak var entryView : UIView!
     @IBOutlet weak var entryViewYConstraint : NSLayoutConstraint!
 
@@ -29,12 +31,18 @@ class StartViewController : BaseViewController {
     private var tapGesture : UITapGestureRecognizer!
     private var oscConnection : OSCTCPConnection?
     
-    private var liveLink : LiveLinkProvider?
+    public var multicastWatchdogSocket  : GCDAsyncUdpSocket?
+    
+    private var liveLink : LiveLink?
     private var liveLinkTimer : Timer?
     
     @objc dynamic let appSettings = AppSettings.shared
     private var observers = [NSKeyValueObservation]()
 
+    var ipAddressIsDemoMode : Bool {
+        self.ipAddress.text == "demo.mode"
+    }
+    
     override var preferredStatusBarStyle: UIStatusBarStyle {
           return .lightContent
     }
@@ -43,6 +51,8 @@ class StartViewController : BaseViewController {
         
         super.viewDidLoad();
         
+        self.restartView.isHidden = true
+        
         if let infoDict = Bundle.main.infoDictionary {
             self.versionLabel.text = String(format: "v%@ (%@)", infoDict["CFBundleShortVersionString"] as! String, infoDict["CFBundleVersion"] as! String)
         } else {
@@ -96,14 +106,71 @@ class StartViewController : BaseViewController {
         
         // any change to the subject name will remove & re-add the camera subject.
         observers.append(observe(\.appSettings.liveLinkSubjectName, options: [.old,.new], changeHandler: { object, change in
-            self.liveLink?.removeCameraSubject(change.oldValue)
+            self.liveLink?.removeCameraSubject(change.oldValue!)
             self.liveLink?.addCameraSubject(self.appSettings.liveLinkSubjectName)
         }))
+        observers.append(observe(\.appSettings.engineVersion, options: [.old,.new], changeHandler: { object, change in
+            if LiveLink.initialized {
+                self.restartView.isHidden = !LiveLink.requiresRestart
+            }
+        }))
+
+
+    }
+    
+    func restartLiveLink() {
+
+        do {
+            try LiveLink.initialize(self)
+
+        } catch LiveLinkError.unknownEngineVersion {
+
+            // if there is no engine version (first run), then we ask the user to choose.
+            let askVersionAlert = UIAlertController(title: NSLocalizedString("version-title", value:"Engine Version", comment: "Title of an settings screen."),
+                                                    message: NSLocalizedString("version-choose", value:"Choose the version of Unreal Engine you will be connecting to. This selection can be modified later in the app's settings.", comment: "Message for seelecting the correct version of unreal engine."), preferredStyle: .alert)
+            askVersionAlert.addAction(UIAlertAction(title: "Unreal Engine 4", style: .default, handler: { _ in
+                LiveLink.engineVersion = .ue4
+                self.restartLiveLink()
+            }))
+            askVersionAlert.addAction(UIAlertAction(title: "Unreal Engine 5", style: .default, handler: { _ in
+                LiveLink.engineVersion = .ue5
+                self.restartLiveLink()
+            }))
+            self.present(askVersionAlert, animated: true)
+            return
+        } catch {
+            
+        }
         
-        LiveLink.initialize(self)
-        self.liveLink = LiveLink.createProvider("VCAM IOS")
-        self.liveLink?.addCameraSubject(AppSettings.shared.liveLinkSubjectName);
+        // stop the provider & restart livelink here
+        if self.liveLink != nil {
+            Log.info("Restarting Messaging Engine.")
+            try? LiveLink.restart()
+            self.liveLink = nil
+        }
+
+        Log.info("Initializing LiveLink Provider.")
 
+        self.liveLink = try? LiveLink("VCAM IOS")
+        self.liveLink?.addCameraSubject(AppSettings.shared.liveLinkSubjectName)
+
+        if let videoViewController = self.presentedViewController as? VideoViewController {
+            if !self.ipAddressIsDemoMode {
+                videoViewController.liveLink = self.liveLink
+            }
+        }
+
+        multicastWatchdogSocket?.close()
+        Log.info("Starting multicast watchdog.")
+        multicastWatchdogSocket = GCDAsyncUdpSocket(delegate: self, delegateQueue: DispatchQueue.main)
+        do {
+            try multicastWatchdogSocket?.enableReusePort(true)
+            try multicastWatchdogSocket?.bind(toPort: 6665)
+            try multicastWatchdogSocket?.joinMulticastGroup("230.0.0.1")
+            try multicastWatchdogSocket?.beginReceiving()
+        } catch {
+            Log.info("Error creating watchdog : \(error.localizedDescription)")
+        }
     }
     
     override func viewWillAppear(_ animated : Bool) {
@@ -113,8 +180,12 @@ class StartViewController : BaseViewController {
         
         liveLinkTimer?.invalidate()
         liveLinkTimer = Timer.scheduledTimer(withTimeInterval: 1.0/10.0, repeats: true, block: { timer in
-            self.liveLink?.updateSubject(AppSettings.shared.liveLinkSubjectName, withTransform: simd_float4x4(), atTime: Timecode.create().toTimeInterval())
+            self.liveLink?.updateSubject(AppSettings.shared.liveLinkSubjectName, transform: simd_float4x4(), time: Timecode.create().toTimeInterval())
         })
+        
+        if !LiveLink.initialized {
+            restartLiveLink();
+        }
     }
     
     override func viewDidDisappear(_ animated: Bool) {
@@ -186,7 +257,7 @@ class StartViewController : BaseViewController {
         
         AppSettings.shared.lastConnectionAddress = self.ipAddress.text!
 
-        if self.ipAddress.text == "demo.mode" {
+        if self.ipAddressIsDemoMode {
             
             self.performSegue(withIdentifier: "showVideoViewDemoMode", sender: self)
 
@@ -255,5 +326,22 @@ extension StartViewController : UITextFieldDelegate {
         
         return true
     }
-    
+}
+
+extension StartViewController : GCDAsyncUdpSocketDelegate {
+
+    func udpSocketDidClose(_ sock: GCDAsyncUdpSocket, withError error: Error?) {
+        Log.error("Multicast watchdog closed : restarting LiveLink.")
+        restartLiveLink()
+    }
+}
+
+extension StartViewController : UE5LiveLinkLogDelegate, UE4LiveLinkLogDelegate {
+
+    func ue5LogMessage(_ message: String!) {
+        Log.info(message)
+    }
+    func ue4LogMessage(_ message: String!) {
+        Log.info(message)
+    }
 }
diff --git a/Engine/Extras/VirtualProduction/LiveLinkVCAM/vcam/VCAM-Bridging-Header.h b/Engine/Extras/VirtualProduction/LiveLinkVCAM/vcam/VCAM-Bridging-Header.h
--- a/Engine/Extras/VirtualProduction/LiveLinkVCAM/vcam/VCAM-Bridging-Header.h
+++ b/Engine/Extras/VirtualProduction/LiveLinkVCAM/vcam/VCAM-Bridging-Header.h
@@ -6,6 +6,7 @@
 
 #import "PassthroughShaderTypes.h"
 #import "ObjCUtility.h"
-#import <LiveLinkProvider/LiveLink.h>
+#import <UE5LiveLinkProvider/UE5LiveLink.h>
+#import <UE4LiveLinkProvider/UE4LiveLink.h>
 
 #endif /* VCAM_Bridging_Header_h */
diff --git a/Engine/Extras/VirtualProduction/LiveLinkVCAM/vcam/VideoViewController+ARSCNViewDelegate.swift b/Engine/Extras/VirtualProduction/LiveLinkVCAM/vcam/VideoViewController+ARSCNViewDelegate.swift
--- a/Engine/Extras/VirtualProduction/LiveLinkVCAM/vcam/VideoViewController+ARSCNViewDelegate.swift
+++ b/Engine/Extras/VirtualProduction/LiveLinkVCAM/vcam/VideoViewController+ARSCNViewDelegate.swift
@@ -13,7 +13,7 @@ extension VideoViewController : ARSCNViewDelegate {
     func renderer(_ renderer: SCNSceneRenderer, updateAtTime time: TimeInterval) {
         
         if let pov = renderer.pointOfView {
-            self.liveLink?.updateSubject(AppSettings.shared.liveLinkSubjectName, withTransform: pov.simdTransform, atTime: Timecode.create().toTimeInterval())
+            self.liveLink?.updateSubject(AppSettings.shared.liveLinkSubjectName, transform: pov.simdTransform, time: Timecode.create().toTimeInterval())
         }
         
     }
diff --git a/Engine/Extras/VirtualProduction/LiveLinkVCAM/vcam/VideoViewController.swift b/Engine/Extras/VirtualProduction/LiveLinkVCAM/vcam/VideoViewController.swift
--- a/Engine/Extras/VirtualProduction/LiveLinkVCAM/vcam/VideoViewController.swift
+++ b/Engine/Extras/VirtualProduction/LiveLinkVCAM/vcam/VideoViewController.swift
@@ -35,7 +35,7 @@ class VideoViewController : BaseViewController {
 
     var oscConnection : OSCTCPConnection?
 
-    weak var liveLink : LiveLinkProvider?
+    weak var liveLink : LiveLink?
     private var decoder : VideoDecoder?
     var dismissOnDisconnect = false
     
