Live Link VCAM (iOS) v1.0.2 build 4 : UI updates per ryan.mayeda's feedback.

diff --git a/Engine/Extras/VirtualProduction/LiveLinkVCAM/vcam/LiveLinkViewController.swift b/Engine/Extras/VirtualProduction/LiveLinkVCAM/vcam/LiveLinkViewController.swift
new file mode 100644
--- /dev/null
+++ b/Engine/Extras/VirtualProduction/LiveLinkVCAM/vcam/LiveLinkViewController.swift
@@ -0,0 +1,100 @@
+//
+//  LiveLinkViewController.swift
+//  vcam
+//
+//  Created by Brian Smith on 1/25/22.
+//  Copyright Epic Games, Inc. All Rights Reserved.
+//
+
+import UIKit
+
+class LiveLinkViewController : UITableViewController {
+    
+    let appSettings = AppSettings.shared
+    
+    override func viewWillAppear(_ animated: Bool) {
+        super.viewWillAppear(animated)
+        self.tableView.reloadData()
+    }
+
+    override func tableView(_ tableView: UITableView, willDisplay cell: UITableViewCell, forRowAt indexPath: IndexPath) {
+        
+        if let detail = cell.detailTextLabel {
+            detail.textColor = UIColor.secondaryLabel
+        }
+        
+        switch cell.reuseIdentifier {
+
+        case "subjectName":
+            cell.detailTextLabel?.text = appSettings.liveLinkSubjectName
+
+        case "protocol":
+            cell.detailTextLabel?.text = (LiveLink.engineVersion == .ue4 ? "4.27 and Earlier" : "5.0 and Later") + (LiveLink.requiresRestart ? " (restart required)" : "")
+
+        default:
+            break
+        }
+    }
+    
+    override func prepare(for segue: UIStoryboardSegue, sender: Any?) {
+
+        if let vc = segue.destination as? SingleValueViewController {
+            
+            if segue.identifier == "subjectName" {
+                
+                vc.navigationItem.title = Localized.subjectName()
+                vc.mode = .edit
+                vc.allowedType = .unreal
+                vc.initialValue = AppSettings.shared.liveLinkSubjectName
+                vc.placeholderValue = AppSettings.defaultLiveLinkSubjectName()
+                vc.finished = { (action, value) in
+                    
+                    if action == .done {
+                        
+                        let v = value!.trimmed()
+                        
+                        AppSettings.shared.liveLinkSubjectName = v.isEmpty ? AppSettings.defaultLiveLinkSubjectName() : value!.toUnrealCompatibleString()
+                        self.tableView.reloadData()
+                    }
+                }
+            }
+        } else if let vc = segue.destination as? MultipleChoiceViewController {
+            
+            if segue.identifier == "protocol" {
+                vc.navigationItem.title = NSLocalizedString("settings-title-protocol", value:"Protocol", comment: "Title of an settings screen.")
+                vc.items = [ "4.27 and Earlier", "5.0 and Later" ]
+                vc.selectedIndex = (LiveLink.engineVersion == .ue4) ? 0 : 1
+                
+                vc.selectionChanged = { (index) in
+
+                    switch index {
+                    case 0:
+                        LiveLink.engineVersion = .ue4
+                    case 1:
+                        LiveLink.engineVersion = .ue5
+                    default:
+                        break
+                    }
+
+                    // if the engine version has changed, prompt the user to restart
+                    if (LiveLink.requiresRestart) {
+                        let errorAlert = UIAlertController(title: NSLocalizedString("alert-title-restart", value: "Restart required", comment: "Title indicating the user must restart the application."), message: NSLocalizedString("alert-message-restart", value: "For this setting to take effect, the application must be terminated and restarted.", comment: "Message indicating the user must restart the application."), preferredStyle: .alert)
+                        errorAlert.addAction(UIAlertAction(title: Localized.buttonOK(), style: .default))
+                        errorAlert.addAction(UIAlertAction(title: NSLocalizedString("alert-button-learnmore", value: "Learn more...", comment: "Button to get more information."), style: .default, handler: { action in
+                            if UIDevice.current.userInterfaceIdiom == .phone {
+                                UIApplication.shared.open(URL(string: "https://support.apple.com/HT201330")!, options: [:])
+                            } else {
+                                UIApplication.shared.open(URL(string: "https://support.apple.com/HT212063")!, options: [:])
+                            }
+                        }))
+                        vc.present(errorAlert, animated: true)
+
+                    }
+                    
+                }
+                
+                
+            }
+        }
+    }
+}
diff --git a/Engine/Extras/VirtualProduction/LiveLinkVCAM/vcam/SettingsViewController.swift b/Engine/Extras/VirtualProduction/LiveLinkVCAM/vcam/SettingsViewController.swift
--- a/Engine/Extras/VirtualProduction/LiveLinkVCAM/vcam/SettingsViewController.swift
+++ b/Engine/Extras/VirtualProduction/LiveLinkVCAM/vcam/SettingsViewController.swift
@@ -33,83 +33,17 @@ class SettingsViewController : UITableViewController {
         
         switch cell.reuseIdentifier {
 
-        case "subjectName":
+        case "liveLink":
             cell.detailTextLabel?.text = appSettings.liveLinkSubjectName
 
         case "timecode":
             cell.detailTextLabel?.text = Timecode.sourceToString(appSettings.timecodeSourceEnum())
 
-        case "engineVersion":
-            cell.detailTextLabel?.text = (LiveLink.engineVersion == .ue4 ? "UE4" : "UE5") + (LiveLink.requiresRestart ? " (restart required)" : "")
-
         default:
             break
         }
     }
     
-    override func prepare(for segue: UIStoryboardSegue, sender: Any?) {
-
-        if let vc = segue.destination as? SingleValueViewController {
-            
-            if segue.identifier == "subjectName" {
-                
-                vc.navigationItem.title = Localized.subjectName()
-                vc.mode = .edit
-                vc.allowedType = .unreal
-                vc.initialValue = AppSettings.shared.liveLinkSubjectName
-                vc.placeholderValue = AppSettings.defaultLiveLinkSubjectName()
-                vc.finished = { (action, value) in
-                    
-                    if action == .done {
-                        
-                        let v = value!.trimmed()
-                        
-                        AppSettings.shared.liveLinkSubjectName = v.isEmpty ? AppSettings.defaultLiveLinkSubjectName() : value!.toUnrealCompatibleString()
-                        self.tableView.reloadData()
-                    }
-                }
-            }
-        } else if let vc = segue.destination as? MultipleChoiceViewController {
-            
-            if segue.identifier == "engineVersion" {
-                vc.navigationItem.title = NSLocalizedString("settings-title-engine-version", value:"Engine Version", comment: "Title of an settings screen.")
-                vc.items = [ "Unreal Engine 4", "Unreal Engine 5" ]
-                vc.selectedIndex = (LiveLink.engineVersion == .ue4) ? 0 : 1
-                vc.footerString = NSLocalizedString("settings-footer-engine-version", value: "The Live Link protocol was updated in Unreal Engine 5.0. Live Link VCAM supports both protocols, but this must be specified to match the version of Unreal Engine in use for the project in order to receive tracking data. After changing this setting the app must be terminated & restarted for the setting to take effect.", comment: "A note about how the LiveLink protocol has changed between UE4 & UE5.")
-                
-                vc.selectionChanged = { (index) in
-
-                    switch index {
-                    case 0:
-                        LiveLink.engineVersion = .ue4
-                    case 1:
-                        LiveLink.engineVersion = .ue5
-                    default:
-                        break
-                    }
-
-                    // if the engine version has changed, prompt the user to restart
-                    if (LiveLink.requiresRestart) {
-                        let errorAlert = UIAlertController(title: NSLocalizedString("alert-title-restart", value: "Restart required", comment: "Title indicating the user must restart the application."), message: NSLocalizedString("alert-message-restart", value: "For this setting to take effect, the application must be terminated and restarted.", comment: "Message indicating the user must restart the application."), preferredStyle: .alert)
-                        errorAlert.addAction(UIAlertAction(title: Localized.buttonOK(), style: .default))
-                        errorAlert.addAction(UIAlertAction(title: NSLocalizedString("alert-button-learnmore", value: "Learn more...", comment: "Button to get more information."), style: .default, handler: { action in
-                            if UIDevice.current.userInterfaceIdiom == .phone {
-                                UIApplication.shared.open(URL(string: "https://support.apple.com/HT201330")!, options: [:])
-                            } else {
-                                UIApplication.shared.open(URL(string: "https://support.apple.com/HT212063")!, options: [:])
-                            }
-                        }))
-                        vc.present(errorAlert, animated: true)
-
-                    }
-                    
-                }
-                
-                
-            }
-        }
-    }
-    
     @objc func done(sender:Any?) {
         self.navigationController?.dismiss(animated: true, completion: nil)
     }
diff --git a/Engine/Extras/VirtualProduction/LiveLinkVCAM/vcam/StartViewController.swift b/Engine/Extras/VirtualProduction/LiveLinkVCAM/vcam/StartViewController.swift
--- a/Engine/Extras/VirtualProduction/LiveLinkVCAM/vcam/StartViewController.swift
+++ b/Engine/Extras/VirtualProduction/LiveLinkVCAM/vcam/StartViewController.swift
@@ -182,6 +182,10 @@ class StartViewController : BaseViewController {
         liveLinkTimer = Timer.scheduledTimer(withTimeInterval: 1.0/10.0, repeats: true, block: { timer in
             self.liveLink?.updateSubject(AppSettings.shared.liveLinkSubjectName, transform: simd_float4x4(), time: Timecode.create().toTimeInterval())
         })
+    }
+    
+    override func viewDidAppear(_ animated: Bool) {
+        super.viewDidAppear(animated)
         
         if !LiveLink.initialized {
             restartLiveLink();
