Live Link VCAM (iOS) v1.0.2 Build 6 : initial support for hardware game controllers (plugged in only,  not bluetooth).

diff --git a/Engine/Extras/VirtualProduction/LiveLinkVCAM/vcam/Info.plist b/Engine/Extras/VirtualProduction/LiveLinkVCAM/vcam/Info.plist
--- a/Engine/Extras/VirtualProduction/LiveLinkVCAM/vcam/Info.plist
+++ b/Engine/Extras/VirtualProduction/LiveLinkVCAM/vcam/Info.plist
@@ -18,6 +18,15 @@
 	<string>$(MARKETING_VERSION)</string>
 	<key>CFBundleVersion</key>
 	<string>$(CURRENT_PROJECT_VERSION)</string>
+	<key>GCSupportedGameControllers</key>
+	<array>
+		<dict>
+			<key>ProfileName</key>
+			<string>ExtendedGamepad</string>
+		</dict>
+	</array>
+	<key>GCSupportsControllerUserInteraction</key>
+	<true/>
 	<key>ITSAppUsesNonExemptEncryption</key>
 	<false/>
 	<key>LSRequiresIPhoneOS</key>
diff --git a/Engine/Extras/VirtualProduction/LiveLinkVCAM/vcam/OSC/OSCAddressPattern.swift b/Engine/Extras/VirtualProduction/LiveLinkVCAM/vcam/OSC/OSCAddressPattern.swift
--- a/Engine/Extras/VirtualProduction/LiveLinkVCAM/vcam/OSC/OSCAddressPattern.swift
+++ b/Engine/Extras/VirtualProduction/LiveLinkVCAM/vcam/OSC/OSCAddressPattern.swift
@@ -20,5 +20,8 @@ enum OSCAddressPattern : String {
     case touchStarted = "/MessageHandler/OnTouchStarted"
     case touchMoved = "/MessageHandler/OnTouchMoved"
     case touchEnded = "/MessageHandler/OnTouchEnded"
+    case controllerAnalog = "/MessageHandler/OnControllerAnalog"
+    case controllerButtonPressed = "/MessageHandler/OnControllerButtonPressed"
+    case controllerButtonReleased = "/MessageHandler/OnControllerButtonReleased"
 }
 
diff --git a/Engine/Extras/VirtualProduction/LiveLinkVCAM/vcam/OSC/OSCPacket.swift b/Engine/Extras/VirtualProduction/LiveLinkVCAM/vcam/OSC/OSCPacket.swift
--- a/Engine/Extras/VirtualProduction/LiveLinkVCAM/vcam/OSC/OSCPacket.swift
+++ b/Engine/Extras/VirtualProduction/LiveLinkVCAM/vcam/OSC/OSCPacket.swift
@@ -82,7 +82,12 @@ class OSCPacket {
         
         return data
     }
-    
+
+    class func encode(_ i : Int8) -> Data {
+
+        withUnsafeBytes(of: i) { Data($0) }
+    }
+
     class func encode(_ i : Int32, bigEndian : Bool = true) -> Data {
 
         withUnsafeBytes(of: bigEndian ? i.bigEndian : i) { Data($0) }
diff --git a/Engine/Extras/VirtualProduction/LiveLinkVCAM/vcam/OSC/OSCUtility.swift b/Engine/Extras/VirtualProduction/LiveLinkVCAM/vcam/OSC/OSCUtility.swift
--- a/Engine/Extras/VirtualProduction/LiveLinkVCAM/vcam/OSC/OSCUtility.swift
+++ b/Engine/Extras/VirtualProduction/LiveLinkVCAM/vcam/OSC/OSCUtility.swift
@@ -76,8 +76,36 @@ class OSCUtility {
         data.append(OSCPacket.encode(Float(point.x), bigEndian: false))
         data.append(OSCPacket.encode(Float(point.y), bigEndian: false))
         data.append(OSCPacket.encode(Int32(finger), bigEndian: false))
-        data.append(OSCPacket.encode(Int32(force), bigEndian: false))
+        data.append(OSCPacket.encode(Float(force), bigEndian: false))
 
         return data;
     }
+    
+    class func ueControllerAnalogData(key : String, controller : Int, value : Float) -> Data {
+        
+        var data = Data(capacity: 32);
+        
+        data.append(OSCPacket.encode(Int32(key.count + 1), bigEndian: false))
+        data.append(key.data(using: .utf8)!)
+        data.append(OSCPacket.encode(Int8(0)))
+        data.append(OSCPacket.encode(Int32(controller), bigEndian: false))
+        data.append(OSCPacket.encode(Float(value), bigEndian: false))
+
+        return data;
+        
+    }
+    
+    class func ueControllerButtonData(key : String, controller : Int, isRepeat : Bool) -> Data {
+        
+        var data = Data(capacity: 32);
+        
+        data.append(OSCPacket.encode(Int32(key.count + 1), bigEndian: false))
+        data.append(key.data(using: .utf8)!)
+        data.append(OSCPacket.encode(Int8(0)))
+        data.append(OSCPacket.encode(Int32(controller), bigEndian: false))
+        data.append(OSCPacket.encode(Int8(isRepeat ? 1 : 0)))
+
+        return data;
+        
+    }
 }
diff --git a/Engine/Extras/VirtualProduction/LiveLinkVCAM/vcam/StartViewController.swift b/Engine/Extras/VirtualProduction/LiveLinkVCAM/vcam/StartViewController.swift
--- a/Engine/Extras/VirtualProduction/LiveLinkVCAM/vcam/StartViewController.swift
+++ b/Engine/Extras/VirtualProduction/LiveLinkVCAM/vcam/StartViewController.swift
@@ -38,6 +38,8 @@ class StartViewController : BaseViewController {
     
     @objc dynamic let appSettings = AppSettings.shared
     private var observers = [NSKeyValueObservation]()
+    
+    private var gameController : GCController?
 
     var ipAddressIsDemoMode : Bool {
         self.ipAddress.text == "demo.mode"
@@ -114,10 +116,18 @@ class StartViewController : BaseViewController {
                 self.restartView.isHidden = !LiveLink.requiresRestart
             }
         }))
-
-
+        
+        NotificationCenter.default.addObserver(self, selector: #selector(gameControllerDidConnectNotification), name: .GCControllerDidConnect, object: nil)
+        NotificationCenter.default.addObserver(self, selector: #selector(gameControllerDidDisconnectNotification), name: .GCControllerDidConnect, object: nil)
+        
+        self.gameController = GCController.controllers().first
+        if let gc = self.gameController {
+            if gc.isAttachedToDevice {
+                gc.playerIndex = .index1
+            }
+        }
     }
-    
+     
     func restartLiveLink() {
 
         do {
@@ -214,6 +224,7 @@ class StartViewController : BaseViewController {
                 vc.oscConnection = self.oscConnection
                 
                 vc.liveLink = self.liveLink
+                vc.gameController = self.gameController
 
                 self.oscConnection = nil
             }
@@ -310,6 +321,23 @@ class StartViewController : BaseViewController {
     @IBAction func textFieldChanged(_ sender : Any?) {
         self.connect.isEnabled = !self.ipAddress.text!.isEmpty
     }
+    
+    @objc func gameControllerDidConnectNotification(_ notification: NSNotification) {
+        
+        self.gameController = notification.object as? GCController
+        self.gameController?.playerIndex = .index1
+        
+        if let videoViewController = self.presentedViewController as? VideoViewController {
+            videoViewController.gameController = self.gameController
+        }
+    }
+
+    @objc func gameControllerDidDisconnectNotification(_ notification: NSNotification) {
+        if let gc = notification.object as? GCController {
+            Log.info("gameControllerDidDisconnectNotification \(gc.vendorName ?? "Unknown controller")")
+        }
+    }
+
 }
 
 extension StartViewController : UIGestureRecognizerDelegate {
diff --git a/Engine/Extras/VirtualProduction/LiveLinkVCAM/vcam/VideoViewController.swift b/Engine/Extras/VirtualProduction/LiveLinkVCAM/vcam/VideoViewController.swift
--- a/Engine/Extras/VirtualProduction/LiveLinkVCAM/vcam/VideoViewController.swift
+++ b/Engine/Extras/VirtualProduction/LiveLinkVCAM/vcam/VideoViewController.swift
@@ -10,6 +10,7 @@ import UIKit
 import MetalKit
 import ARKit
 import Easing
+import GameController
 
 class VideoViewController : BaseViewController {
 
@@ -42,6 +43,30 @@ class VideoViewController : BaseViewController {
     var stats = Array<(Date,Int)>()
     var statsTimer : Timer?
     
+    weak var gameController : GCController? {
+        didSet {
+            
+            if let gc = gameController {
+                gc.extendedGamepad?.leftThumbstick.valueChangedHandler = self.directionalPadValueChanged
+                gc.extendedGamepad?.rightThumbstick.valueChangedHandler = self.directionalPadValueChanged
+                gc.extendedGamepad?.dpad.up.valueChangedHandler = self.buttonValueChanged
+                gc.extendedGamepad?.dpad.down.valueChangedHandler = self.buttonValueChanged
+                gc.extendedGamepad?.dpad.left.valueChangedHandler = self.buttonValueChanged
+                gc.extendedGamepad?.dpad.right.valueChangedHandler = self.buttonValueChanged
+                gc.extendedGamepad?.buttonA.valueChangedHandler = self.buttonValueChanged
+                gc.extendedGamepad?.buttonB.valueChangedHandler = self.buttonValueChanged
+                gc.extendedGamepad?.buttonX.valueChangedHandler = self.buttonValueChanged
+                gc.extendedGamepad?.buttonY.valueChangedHandler = self.buttonValueChanged
+                gc.extendedGamepad?.leftShoulder.valueChangedHandler = self.buttonValueChanged
+                gc.extendedGamepad?.rightShoulder.valueChangedHandler = self.buttonValueChanged
+                gc.extendedGamepad?.leftTrigger.valueChangedHandler = self.buttonValueChanged
+                gc.extendedGamepad?.rightTrigger.valueChangedHandler = self.buttonValueChanged
+                gc.extendedGamepad?.leftThumbstickButton?.pressedChangedHandler = self.buttonValueChanged
+                gc.extendedGamepad?.rightThumbstickButton?.pressedChangedHandler = self.buttonValueChanged
+            }
+        }
+    }
+    
     private var currentTouches = Array<UITouch?>()
     
     var relayTouchEvents = true {
@@ -594,6 +619,75 @@ class VideoViewController : BaseViewController {
         // save the actual animating value of the header view's Y position
         self.headerViewY = self.headerView.layer.presentation()?.frame.minY ?? 0
     }
+
+    func directionalPadValueChanged(directionalPad : GCControllerDirectionPad, xValue : Float, yValue : Float) {
+        guard let osc = oscConnection else { return }
+        guard let gc = gameController else { return }
+        guard let gp = gc.extendedGamepad else { return }
+
+        let controllerIndex = gc.playerIndex.rawValue
+        
+        switch directionalPad {
+        case gp.leftThumbstick:
+            osc.send(.controllerAnalog, arguments: [ OSCArgument.blob(OSCUtility.ueControllerAnalogData(key: "Gamepad_LeftX", controller: controllerIndex, value: xValue)) ])
+            osc.send(.controllerAnalog, arguments: [ OSCArgument.blob(OSCUtility.ueControllerAnalogData(key: "Gamepad_LeftY", controller: controllerIndex, value: yValue)) ])
+        case gp.rightThumbstick:
+            osc.send(.controllerAnalog, arguments: [ OSCArgument.blob(OSCUtility.ueControllerAnalogData(key: "Gamepad_RightX", controller: controllerIndex, value: xValue)) ])
+            osc.send(.controllerAnalog, arguments: [ OSCArgument.blob(OSCUtility.ueControllerAnalogData(key: "Gamepad_RightY", controller: controllerIndex, value: yValue)) ])
+        default:
+            Log.warning("directionalPadValueChanged() encounted unsupported '\(directionalPad.localizedName ?? "UNK")'")
+            break
+        }
+    }
+    
+    func buttonValueChanged(button : GCControllerButtonInput, value : Float, pressed : Bool) {
+
+        guard let osc = oscConnection else { return }
+        guard let gc = gameController else { return }
+        guard let gp = gc.extendedGamepad else { return }
+
+        let address = button.isPressed ? OSCAddressPattern.controllerButtonPressed : OSCAddressPattern.controllerButtonReleased
+        let controllerIndex = gc.playerIndex.rawValue
+        let isRepeat = false
+
+        switch button {
+        case gp.leftThumbstickButton:
+            osc.send(address, arguments: [ OSCArgument.blob(OSCUtility.ueControllerButtonData(key: "Gamepad_LeftThumbstick", controller: controllerIndex, isRepeat: isRepeat)) ])
+        case gp.rightThumbstickButton:
+            osc.send(address, arguments: [ OSCArgument.blob(OSCUtility.ueControllerButtonData(key: "Gamepad_RightThumbstick", controller: controllerIndex, isRepeat: isRepeat)) ])
+
+        case gp.buttonA:
+            osc.send(address, arguments: [ OSCArgument.blob(OSCUtility.ueControllerButtonData(key: "Gamepad_FaceButton_Bottom", controller: controllerIndex, isRepeat: isRepeat)) ])
+        case gp.buttonB:
+            osc.send(address, arguments: [ OSCArgument.blob(OSCUtility.ueControllerButtonData(key: "Gamepad_FaceButton_Right", controller: controllerIndex, isRepeat: isRepeat)) ])
+        case gp.buttonX:
+            osc.send(address, arguments: [ OSCArgument.blob(OSCUtility.ueControllerButtonData(key: "Gamepad_FaceButton_Left", controller: controllerIndex, isRepeat: isRepeat)) ])
+        case gp.buttonY:
+            osc.send(address, arguments: [ OSCArgument.blob(OSCUtility.ueControllerButtonData(key: "Gamepad_FaceButton_Top", controller: controllerIndex, isRepeat: isRepeat)) ])
+
+        case gp.leftShoulder:
+            osc.send(address, arguments: [ OSCArgument.blob(OSCUtility.ueControllerButtonData(key: "Gamepad_LeftShoulder", controller: controllerIndex, isRepeat: isRepeat)) ])
+        case gp.rightShoulder:
+            osc.send(address, arguments: [ OSCArgument.blob(OSCUtility.ueControllerButtonData(key: "Gamepad_RightShoulder", controller: controllerIndex, isRepeat: isRepeat)) ])
+
+        case gp.leftTrigger:
+            osc.send(address, arguments: [ OSCArgument.blob(OSCUtility.ueControllerButtonData(key: "Gamepad_LeftTrigger", controller: controllerIndex, isRepeat: isRepeat)) ])
+        case gp.rightTrigger:
+            osc.send(address, arguments: [ OSCArgument.blob(OSCUtility.ueControllerButtonData(key: "Gamepad_RightTrigger", controller: controllerIndex, isRepeat: isRepeat)) ])
+
+        case gp.dpad.up:
+            osc.send(address, arguments: [ OSCArgument.blob(OSCUtility.ueControllerButtonData(key: "Gamepad_DPad_Up", controller: controllerIndex, isRepeat: isRepeat)) ])
+        case gp.dpad.down:
+            osc.send(address, arguments: [ OSCArgument.blob(OSCUtility.ueControllerButtonData(key: "Gamepad_DPad_Down", controller: controllerIndex, isRepeat: isRepeat)) ])
+        case gp.dpad.left:
+            osc.send(address, arguments: [ OSCArgument.blob(OSCUtility.ueControllerButtonData(key: "Gamepad_DPad_Left", controller: controllerIndex, isRepeat: isRepeat)) ])
+        case gp.dpad.right:
+            osc.send(address, arguments: [ OSCArgument.blob(OSCUtility.ueControllerButtonData(key: "Gamepad_DPad_Right", controller: controllerIndex, isRepeat: isRepeat)) ])
+        default:
+            Log.warning("buttonValueChanged() encounted unsupported '\(button.localizedName ?? "UNK")'")
+        }
+
+    }
 }
 
 extension VideoViewController : HeaderViewDelegate {
