Live Link VCAM (iOS) v1.0.2 build 7/8 : Fixed various HW controller issues per partner reports.

diff --git a/Engine/Extras/VirtualProduction/LiveLinkVCAM/vcam/StartViewController.swift b/Engine/Extras/VirtualProduction/LiveLinkVCAM/vcam/StartViewController.swift
--- a/Engine/Extras/VirtualProduction/LiveLinkVCAM/vcam/StartViewController.swift
+++ b/Engine/Extras/VirtualProduction/LiveLinkVCAM/vcam/StartViewController.swift
@@ -118,7 +118,7 @@ class StartViewController : BaseViewController {
         }))
         
         NotificationCenter.default.addObserver(self, selector: #selector(gameControllerDidConnectNotification), name: .GCControllerDidConnect, object: nil)
-        NotificationCenter.default.addObserver(self, selector: #selector(gameControllerDidDisconnectNotification), name: .GCControllerDidConnect, object: nil)
+        NotificationCenter.default.addObserver(self, selector: #selector(gameControllerDidDisconnectNotification), name: .GCControllerDidDisconnect, object: nil)
         
         self.gameController = GCController.controllers().first
         if let gc = self.gameController {
diff --git a/Engine/Extras/VirtualProduction/LiveLinkVCAM/vcam/VideoViewController+ARSCNViewDelegate.swift b/Engine/Extras/VirtualProduction/LiveLinkVCAM/vcam/VideoViewController+ARSCNViewDelegate.swift
--- a/Engine/Extras/VirtualProduction/LiveLinkVCAM/vcam/VideoViewController+ARSCNViewDelegate.swift
+++ b/Engine/Extras/VirtualProduction/LiveLinkVCAM/vcam/VideoViewController+ARSCNViewDelegate.swift
@@ -16,6 +16,9 @@ extension VideoViewController : ARSCNViewDelegate {
             self.liveLink?.updateSubject(AppSettings.shared.liveLinkSubjectName, transform: pov.simdTransform, time: Timecode.create().toTimeInterval())
         }
         
+        // update controller
+        self.sendControllerUpdate()
+        
     }
     
     func renderer(_ renderer: SCNSceneRenderer, didRenderScene scene: SCNScene, atTime time: TimeInterval) {
diff --git a/Engine/Extras/VirtualProduction/LiveLinkVCAM/vcam/VideoViewController.swift b/Engine/Extras/VirtualProduction/LiveLinkVCAM/vcam/VideoViewController.swift
--- a/Engine/Extras/VirtualProduction/LiveLinkVCAM/vcam/VideoViewController.swift
+++ b/Engine/Extras/VirtualProduction/LiveLinkVCAM/vcam/VideoViewController.swift
@@ -43,26 +43,28 @@ class VideoViewController : BaseViewController {
     var stats = Array<(Date,Int)>()
     var statsTimer : Timer?
     
+    var gameControllerSnapshot : GCController?
     weak var gameController : GCController? {
         didSet {
-            
             if let gc = gameController {
-                gc.extendedGamepad?.leftThumbstick.valueChangedHandler = self.directionalPadValueChanged
-                gc.extendedGamepad?.rightThumbstick.valueChangedHandler = self.directionalPadValueChanged
-                gc.extendedGamepad?.dpad.up.valueChangedHandler = self.buttonValueChanged
-                gc.extendedGamepad?.dpad.down.valueChangedHandler = self.buttonValueChanged
-                gc.extendedGamepad?.dpad.left.valueChangedHandler = self.buttonValueChanged
-                gc.extendedGamepad?.dpad.right.valueChangedHandler = self.buttonValueChanged
-                gc.extendedGamepad?.buttonA.valueChangedHandler = self.buttonValueChanged
-                gc.extendedGamepad?.buttonB.valueChangedHandler = self.buttonValueChanged
-                gc.extendedGamepad?.buttonX.valueChangedHandler = self.buttonValueChanged
-                gc.extendedGamepad?.buttonY.valueChangedHandler = self.buttonValueChanged
-                gc.extendedGamepad?.leftShoulder.valueChangedHandler = self.buttonValueChanged
-                gc.extendedGamepad?.rightShoulder.valueChangedHandler = self.buttonValueChanged
-                gc.extendedGamepad?.leftTrigger.valueChangedHandler = self.buttonValueChanged
-                gc.extendedGamepad?.rightTrigger.valueChangedHandler = self.buttonValueChanged
+                //gc.extendedGamepad?.leftThumbstick.valueChangedHandler = self.directionalPadValueChanged
+                //gc.extendedGamepad?.rightThumbstick.valueChangedHandler = self.directionalPadValueChanged
+                gc.extendedGamepad?.dpad.up.pressedChangedHandler = self.buttonValueChanged
+                gc.extendedGamepad?.dpad.down.pressedChangedHandler = self.buttonValueChanged
+                gc.extendedGamepad?.dpad.left.pressedChangedHandler = self.buttonValueChanged
+                gc.extendedGamepad?.dpad.right.pressedChangedHandler = self.buttonValueChanged
+                gc.extendedGamepad?.buttonA.pressedChangedHandler = self.buttonValueChanged
+                gc.extendedGamepad?.buttonB.pressedChangedHandler = self.buttonValueChanged
+                gc.extendedGamepad?.buttonX.pressedChangedHandler = self.buttonValueChanged
+                gc.extendedGamepad?.buttonY.pressedChangedHandler = self.buttonValueChanged
+                gc.extendedGamepad?.leftShoulder.pressedChangedHandler = self.buttonValueChanged
+                gc.extendedGamepad?.rightShoulder.pressedChangedHandler = self.buttonValueChanged
+                gc.extendedGamepad?.leftTrigger.pressedChangedHandler = self.buttonValueChanged
+                gc.extendedGamepad?.rightTrigger.pressedChangedHandler = self.buttonValueChanged
                 gc.extendedGamepad?.leftThumbstickButton?.pressedChangedHandler = self.buttonValueChanged
                 gc.extendedGamepad?.rightThumbstickButton?.pressedChangedHandler = self.buttonValueChanged
+                gc.extendedGamepad?.buttonOptions?.pressedChangedHandler = self.buttonValueChanged
+                gc.extendedGamepad?.buttonMenu.pressedChangedHandler = self.buttonValueChanged
             }
         }
     }
@@ -619,6 +621,33 @@ class VideoViewController : BaseViewController {
         // save the actual animating value of the header view's Y position
         self.headerViewY = self.headerView.layer.presentation()?.frame.minY ?? 0
     }
+    
+    func sendThumbstickUpdate(key : String, controller : Int, oldValue : Float?, newValue : Float) {
+
+        guard let osc = oscConnection else { return }
+
+        if (newValue != 0.0) || ((oldValue ?? 0.0) != newValue) {
+            osc.send(.controllerAnalog, arguments: [ OSCArgument.blob(OSCUtility.ueControllerAnalogData(key: key, controller: controller, value: newValue)) ])
+        }
+    }
+    
+    func sendControllerUpdate() {
+        
+        guard let _ = oscConnection else { return }
+        guard let gc = gameController else { return }
+        
+        let snapshot = gc.capture()
+        guard let gp = snapshot.extendedGamepad else { return }
+
+        let controllerIndex = gc.playerIndex.rawValue
+        
+        sendThumbstickUpdate(key: "Gamepad_LeftX", controller: controllerIndex, oldValue : self.gameControllerSnapshot?.extendedGamepad?.leftThumbstick.xAxis.value, newValue :gp.leftThumbstick.xAxis.value)
+        sendThumbstickUpdate(key: "Gamepad_LeftY", controller: controllerIndex, oldValue : self.gameControllerSnapshot?.extendedGamepad?.leftThumbstick.yAxis.value, newValue :gp.leftThumbstick.yAxis.value)
+        sendThumbstickUpdate(key: "Gamepad_RightX", controller: controllerIndex, oldValue : self.gameControllerSnapshot?.extendedGamepad?.rightThumbstick.xAxis.value, newValue :gp.rightThumbstick.xAxis.value)
+        sendThumbstickUpdate(key: "Gamepad_RightY", controller: controllerIndex, oldValue : self.gameControllerSnapshot?.extendedGamepad?.rightThumbstick.yAxis.value, newValue :gp.rightThumbstick.yAxis.value)
+        
+        self.gameControllerSnapshot = snapshot
+    }
 
     func directionalPadValueChanged(directionalPad : GCControllerDirectionPad, xValue : Float, yValue : Float) {
         guard let osc = oscConnection else { return }
@@ -683,6 +712,12 @@ class VideoViewController : BaseViewController {
             osc.send(address, arguments: [ OSCArgument.blob(OSCUtility.ueControllerButtonData(key: "Gamepad_DPad_Left", controller: controllerIndex, isRepeat: isRepeat)) ])
         case gp.dpad.right:
             osc.send(address, arguments: [ OSCArgument.blob(OSCUtility.ueControllerButtonData(key: "Gamepad_DPad_Right", controller: controllerIndex, isRepeat: isRepeat)) ])
+            
+        case gp.buttonOptions:
+            osc.send(address, arguments: [ OSCArgument.blob(OSCUtility.ueControllerButtonData(key: "Gamepad_Special_Left", controller: controllerIndex, isRepeat: isRepeat)) ])
+        case gp.buttonMenu:
+            osc.send(address, arguments: [ OSCArgument.blob(OSCUtility.ueControllerButtonData(key: "Gamepad_Special_Right", controller: controllerIndex, isRepeat: isRepeat)) ])
+
         default:
             Log.warning("buttonValueChanged() encounted unsupported '\(button.localizedName ?? "UNK")'")
         }
