Color wheel widget

diff --git a/Engine/Extras/VirtualProduction/EpicStageApp/android/app/src/main/AndroidManifest.xml b/Engine/Extras/VirtualProduction/EpicStageApp/android/app/src/main/AndroidManifest.xml
--- a/Engine/Extras/VirtualProduction/EpicStageApp/android/app/src/main/AndroidManifest.xml
+++ b/Engine/Extras/VirtualProduction/EpicStageApp/android/app/src/main/AndroidManifest.xml
@@ -1,7 +1,7 @@
 <manifest xmlns:android="http://schemas.android.com/apk/res/android"
-    package="com.epicgames.epic_stage_app">
+    package="com.epicgames.epicstage">
    <application
-        android:label="epic_stage_app"
+        android:label="epic_stage"
         android:name="${applicationName}"
         android:icon="@mipmap/ic_launcher">
         <activity
