Lightcard selection page

diff --git a/Engine/Extras/VirtualProduction/EpicStageApp/android/app/src/debug/AndroidManifest.xml b/Engine/Extras/VirtualProduction/EpicStageApp/android/app/src/debug/AndroidManifest.xml
--- a/Engine/Extras/VirtualProduction/EpicStageApp/android/app/src/debug/AndroidManifest.xml
+++ b/Engine/Extras/VirtualProduction/EpicStageApp/android/app/src/debug/AndroidManifest.xml
@@ -1,5 +1,5 @@
 <manifest xmlns:android="http://schemas.android.com/apk/res/android"
-    package="com.epicgames.epic_stage_app">
+    package="com.epicgames.epicstage">
     <!-- Flutter needs it to communicate with the running application
          to allow setting breakpoints, to provide hot reload, etc.
     -->
diff --git a/Engine/Extras/VirtualProduction/EpicStageApp/android/app/src/profile/AndroidManifest.xml b/Engine/Extras/VirtualProduction/EpicStageApp/android/app/src/profile/AndroidManifest.xml
--- a/Engine/Extras/VirtualProduction/EpicStageApp/android/app/src/profile/AndroidManifest.xml
+++ b/Engine/Extras/VirtualProduction/EpicStageApp/android/app/src/profile/AndroidManifest.xml
@@ -1,5 +1,5 @@
 <manifest xmlns:android="http://schemas.android.com/apk/res/android"
-    package="com.epicgames.epic_stage_app">
+    package="com.epicgames.epicstage">
     <!-- Flutter needs it to communicate with the running application
          to allow setting breakpoints, to provide hot reload, etc.
     -->
