Epic Stage App: Increase TTL for discovery socket and reopen if closed. Add magic number to beacon. Lock to landscape orientation.

diff --git a/Engine/Extras/VirtualProduction/EpicStageApp/android/app/src/main/AndroidManifest.xml b/Engine/Extras/VirtualProduction/EpicStageApp/android/app/src/main/AndroidManifest.xml
--- a/Engine/Extras/VirtualProduction/EpicStageApp/android/app/src/main/AndroidManifest.xml
+++ b/Engine/Extras/VirtualProduction/EpicStageApp/android/app/src/main/AndroidManifest.xml
@@ -11,7 +11,8 @@
             android:theme="@style/LaunchTheme"
             android:configChanges="orientation|keyboardHidden|keyboard|screenSize|smallestScreenSize|locale|layoutDirection|fontScale|screenLayout|density|uiMode"
             android:hardwareAccelerated="true"
-            android:windowSoftInputMode="adjustResize">
+            android:windowSoftInputMode="adjustResize"
+            android:screenOrientation="landscape">
             <!-- Specifies an Android theme to apply to this Activity as soon as
                  the Android process has started. This theme is visible to the user
                  while the Flutter UI initializes. After that, this theme continues
@@ -31,4 +32,8 @@
             android:name="flutterEmbedding"
             android:value="2" />
     </application>
+
+    <uses-permission android:name="android.permission.INTERNET"/>
+    <uses-permission android:name="android.permission.ACCESS_NETWORK_STATE" />
+    <uses-permission android:name="android.permission.ACCESS_WIFI_STATE" />
 </manifest>
diff --git a/Engine/Plugins/VirtualProduction/RemoteControl/Source/WebRemoteControl/Private/StageAppBeaconReceiver.cpp b/Engine/Plugins/VirtualProduction/RemoteControl/Source/WebRemoteControl/Private/StageAppBeaconReceiver.cpp
--- a/Engine/Plugins/VirtualProduction/RemoteControl/Source/WebRemoteControl/Private/StageAppBeaconReceiver.cpp
+++ b/Engine/Plugins/VirtualProduction/RemoteControl/Source/WebRemoteControl/Private/StageAppBeaconReceiver.cpp
@@ -16,6 +16,7 @@
 namespace StageAppBeaconReceiverConstants
 {
 	const uint8 ProtocolVersion = 0;
+	const TArray<uint8> ProtocolIdentifier = { 'E', 'S', '@', 'p' };
 }
 
 FStageAppBeaconReceiver::FStageAppBeaconReceiver()
@@ -145,11 +146,23 @@ void FStageAppBeaconReceiver::ReceiveBeaconMessages()
 
 void FStageAppBeaconReceiver::HandleBeaconMessage(FArrayReader& MessageData, TSharedRef<FInternetAddr> Source)
 {
-	if (MessageData.Num() != 1)
+	if (MessageData.Num() != 5)
 	{
 		return;
 	}
 
+	// Check that the protocol identifier matches. If not, this message is probably unrelated.
+	for (const uint8 ExpectedByte : StageAppBeaconReceiverConstants::ProtocolIdentifier)
+	{
+		uint8 AppByte;
+		MessageData << AppByte;
+
+		if (AppByte != ExpectedByte)
+		{
+			return;
+		}
+	}
+
 	// We don't do anything with this yet, but could use it to ignore beacons from apps with incompatible protocols
 	uint8 AppProtocolVersion;
 	MessageData << AppProtocolVersion;
