Live Link VCAM iOS (v1.1.0 Build 1) : rolled back UE4/UE5 backend selection due to new UE5 LiveLink backwards compatibility.

diff --git a/Engine/Extras/VirtualProduction/LiveLinkVCAM/LiveLinkProvider.framework/Headers/LiveLink.h b/Engine/Extras/VirtualProduction/LiveLinkVCAM/LiveLinkProvider.framework/Headers/LiveLink.h
new file mode 100644
--- /dev/null
+++ b/Engine/Extras/VirtualProduction/LiveLinkVCAM/LiveLinkProvider.framework/Headers/LiveLink.h
@@ -0,0 +1,22 @@
+
+#import <Foundation/Foundation.h>
+
+#import "LiveLinkProvider.h"
+
+__attribute__((visibility("default")))
+@protocol LiveLinkLogDelegate
+
+- (void) LogMessage:(NSString *)message;
+
+@end
+
+ 
+__attribute__((visibility("default")))
+@interface LiveLink : NSObject
+
++ (void) initialize:(id<LiveLinkLogDelegate>)logDelegate;
++ (void) restart;
++ (void) shutdown;
++ (id<LiveLinkProvider> ) createProvider:(NSString *)name;
+
+@end
diff --git a/Engine/Extras/VirtualProduction/LiveLinkVCAM/LiveLinkProvider.framework/Headers/LiveLinkProvider.h b/Engine/Extras/VirtualProduction/LiveLinkVCAM/LiveLinkProvider.framework/Headers/LiveLinkProvider.h
new file mode 100644
--- /dev/null
+++ b/Engine/Extras/VirtualProduction/LiveLinkVCAM/LiveLinkProvider.framework/Headers/LiveLinkProvider.h
@@ -0,0 +1,14 @@
+
+#import <Foundation/Foundation.h>
+#import <simd/simd.h>
+
+@protocol LiveLinkProvider <NSObject>
+
+- (void)addCameraSubject:(NSString *)subjectName;
+- (void)removeCameraSubject:(NSString *)subjectName;
+
+- (void)updateSubject:(NSString *)subjectName withTransform:(simd_float4x4)xform;
+- (void)updateSubject:(NSString *)subjectName withTransform:(simd_float4x4)xform atTime:(double)time;
+
+@end
+
diff --git a/Engine/Extras/VirtualProduction/LiveLinkVCAM/LiveLinkProvider.framework/Info.plist b/Engine/Extras/VirtualProduction/LiveLinkVCAM/LiveLinkProvider.framework/Info.plist
new file mode 100644
--- /dev/null
+++ b/Engine/Extras/VirtualProduction/LiveLinkVCAM/LiveLinkProvider.framework/Info.plist
@@ -0,0 +1,24 @@
+<?xml version="1.0" encoding="UTF-8"?>
+<!DOCTYPE plist PUBLIC "-//Apple//DTD PLIST 1.0//EN" "http://www.apple.com/DTDs/PropertyList-1.0.dtd">
+<plist version="1.0">
+<dict>
+	<key>CFBundleDevelopmentRegion</key>
+	<string>en</string>
+	<key>MinimumOSVersion</key>
+	<string>14.0</string>
+	<key>CFBundleExecutable</key>
+	<string>LiveLinkProvider</string>
+	<key>CFBundleIdentifier</key>
+	<string>com.epicgames.LiveLinkProvider.framework</string>
+	<key>CFBundleInfoDictionaryVersion</key>
+	<string>6.0</string>
+	<key>CFBundleName</key>
+	<string>LiveLinkProvider</string>
+	<key>CFBundlePackageType</key>
+	<string>FMWK</string>
+	<key>CFBundleShortVersionString</key>
+	<string>1.0</string>
+	<key>CFBundleVersion</key>
+	<string>1</string>
+</dict>
+</plist>
diff --git a/Engine/Extras/VirtualProduction/LiveLinkVCAM/UE4LiveLinkProvider.framework/Headers/UE4LiveLink.h b/Engine/Extras/VirtualProduction/LiveLinkVCAM/UE4LiveLinkProvider.framework/Headers/UE4LiveLink.h
deleted file mode 100644
diff --git a/Engine/Extras/VirtualProduction/LiveLinkVCAM/UE4LiveLinkProvider.framework/Headers/UE4LiveLinkProvider.h b/Engine/Extras/VirtualProduction/LiveLinkVCAM/UE4LiveLinkProvider.framework/Headers/UE4LiveLinkProvider.h
deleted file mode 100644
diff --git a/Engine/Extras/VirtualProduction/LiveLinkVCAM/UE4LiveLinkProvider.framework/Info.plist b/Engine/Extras/VirtualProduction/LiveLinkVCAM/UE4LiveLinkProvider.framework/Info.plist
deleted file mode 100644
diff --git a/Engine/Extras/VirtualProduction/LiveLinkVCAM/UE5LiveLinkProvider.framework/Headers/UE5LiveLink.h b/Engine/Extras/VirtualProduction/LiveLinkVCAM/UE5LiveLinkProvider.framework/Headers/UE5LiveLink.h
deleted file mode 100644
diff --git a/Engine/Extras/VirtualProduction/LiveLinkVCAM/UE5LiveLinkProvider.framework/Headers/UE5LiveLinkProvider.h b/Engine/Extras/VirtualProduction/LiveLinkVCAM/UE5LiveLinkProvider.framework/Headers/UE5LiveLinkProvider.h
deleted file mode 100644
diff --git a/Engine/Extras/VirtualProduction/LiveLinkVCAM/UE5LiveLinkProvider.framework/Info.plist b/Engine/Extras/VirtualProduction/LiveLinkVCAM/UE5LiveLinkProvider.framework/Info.plist
deleted file mode 100644
diff --git a/Engine/Extras/VirtualProduction/LiveLinkVCAM/vcam/AppDelegate.swift b/Engine/Extras/VirtualProduction/LiveLinkVCAM/vcam/AppDelegate.swift
--- a/Engine/Extras/VirtualProduction/LiveLinkVCAM/vcam/AppDelegate.swift
+++ b/Engine/Extras/VirtualProduction/LiveLinkVCAM/vcam/AppDelegate.swift
@@ -30,7 +30,7 @@ class AppDelegate: UIResponder, UIApplicationDelegate {
     }
     
     func applicationWillTerminate(_ application: UIApplication) {
-        try? LiveLink.shutdown()
+        LiveLink.shutdown()
     }
 
     // MARK: UISceneSession Lifecycle
diff --git a/Engine/Extras/VirtualProduction/LiveLinkVCAM/vcam/AppSettings.swift b/Engine/Extras/VirtualProduction/LiveLinkVCAM/vcam/AppSettings.swift
--- a/Engine/Extras/VirtualProduction/LiveLinkVCAM/vcam/AppSettings.swift
+++ b/Engine/Extras/VirtualProduction/LiveLinkVCAM/vcam/AppSettings.swift
@@ -26,7 +26,6 @@ class AppSettings : NSObject {
         static let didAcceptEULA = "didAcceptEULA"
         static let lastConnectionAddress = "lastConnectionAddress"
         
-        static let engineVersion = "engineVersion"
         static let liveLinkSubjectName = "liveLinkSubjectName"
 
         static let timecodeSource = "timecodeSource"
@@ -44,7 +43,6 @@ class AppSettings : NSObject {
             Keys.didAcceptEULA : false,
             Keys.lastConnectionAddress : "",
 
-            Keys.engineVersion : "",
             Keys.liveLinkSubjectName : defaultLiveLinkSubjectName(),
             
             Keys.timecodeSource : 1,
@@ -89,12 +87,6 @@ class AppSettings : NSObject {
         }
     }
 
-    @objc dynamic var engineVersion : String = UserDefaults.standard.string(forKey: Keys.engineVersion) ?? "" {
-        didSet {
-            UserDefaults.standard.set(engineVersion, forKey: Keys.engineVersion)
-        }
-    }
-
     @objc dynamic var liveLinkSubjectName : String = UserDefaults.standard.string(forKey: Keys.liveLinkSubjectName) ?? defaultLiveLinkSubjectName() {
         didSet {
             UserDefaults.standard.set(liveLinkSubjectName, forKey: Keys.liveLinkSubjectName)
diff --git a/Engine/Extras/VirtualProduction/LiveLinkVCAM/vcam/LiveLink.swift b/Engine/Extras/VirtualProduction/LiveLinkVCAM/vcam/LiveLink.swift
deleted file mode 100644
diff --git a/Engine/Extras/VirtualProduction/LiveLinkVCAM/vcam/LiveLinkViewController.swift b/Engine/Extras/VirtualProduction/LiveLinkVCAM/vcam/LiveLinkViewController.swift
deleted file mode 100644
diff --git a/Engine/Extras/VirtualProduction/LiveLinkVCAM/vcam/LogViewController.swift b/Engine/Extras/VirtualProduction/LiveLinkVCAM/vcam/LogViewController.swift
--- a/Engine/Extras/VirtualProduction/LiveLinkVCAM/vcam/LogViewController.swift
+++ b/Engine/Extras/VirtualProduction/LiveLinkVCAM/vcam/LogViewController.swift
@@ -150,7 +150,7 @@ class LogActivityItemSource : NSObject, UIActivityItemSource {
             return url
         } else {
 
-            let url = FileManager.default.temporaryDirectory.appendingPathComponent("LiveLinkFace.log")
+            let url = FileManager.default.temporaryDirectory.appendingPathComponent("LiveLinkVCAM.log")
             if let fileHandle = FileUtility.createFileAtURL(url, overwrite: true) {
                 let logItems = Log.items;
                 for item in logItems {
diff --git a/Engine/Extras/VirtualProduction/LiveLinkVCAM/vcam/SettingsViewController.swift b/Engine/Extras/VirtualProduction/LiveLinkVCAM/vcam/SettingsViewController.swift
--- a/Engine/Extras/VirtualProduction/LiveLinkVCAM/vcam/SettingsViewController.swift
+++ b/Engine/Extras/VirtualProduction/LiveLinkVCAM/vcam/SettingsViewController.swift
@@ -33,7 +33,7 @@ class SettingsViewController : UITableViewController {
         
         switch cell.reuseIdentifier {
 
-        case "liveLink":
+        case "subjectName":
             cell.detailTextLabel?.text = appSettings.liveLinkSubjectName
 
         case "timecode":
@@ -48,4 +48,28 @@ class SettingsViewController : UITableViewController {
         self.navigationController?.dismiss(animated: true, completion: nil)
     }
 
+    override func prepare(for segue: UIStoryboardSegue, sender: Any?) {
+
+        if let vc = segue.destination as? SingleValueViewController {
+            
+            if segue.identifier == "subjectName" {
+                
+                vc.navigationItem.title = Localized.subjectName()
+                vc.mode = .edit
+                vc.allowedType = .unreal
+                vc.initialValue = AppSettings.shared.liveLinkSubjectName
+                vc.placeholderValue = AppSettings.defaultLiveLinkSubjectName()
+                vc.finished = { (action, value) in
+                    
+                    if action == .done {
+                        
+                        let v = value!.trimmed()
+                        
+                        AppSettings.shared.liveLinkSubjectName = v.isEmpty ? AppSettings.defaultLiveLinkSubjectName() : value!.toUnrealCompatibleString()
+                        self.tableView.reloadData()
+                    }
+                }
+            }
+        }
+    }
 }
diff --git a/Engine/Extras/VirtualProduction/LiveLinkVCAM/vcam/StartViewController.swift b/Engine/Extras/VirtualProduction/LiveLinkVCAM/vcam/StartViewController.swift
--- a/Engine/Extras/VirtualProduction/LiveLinkVCAM/vcam/StartViewController.swift
+++ b/Engine/Extras/VirtualProduction/LiveLinkVCAM/vcam/StartViewController.swift
@@ -33,7 +33,7 @@ class StartViewController : BaseViewController {
     
     public var multicastWatchdogSocket  : GCDAsyncUdpSocket?
     
-    private var liveLink : LiveLink?
+    private var liveLink : LiveLinkProvider?
     private var liveLinkTimer : Timer?
     
     @objc dynamic let appSettings = AppSettings.shared
@@ -111,11 +111,6 @@ class StartViewController : BaseViewController {
             self.liveLink?.removeCameraSubject(change.oldValue!)
             self.liveLink?.addCameraSubject(self.appSettings.liveLinkSubjectName)
         }))
-        observers.append(observe(\.appSettings.engineVersion, options: [.old,.new], changeHandler: { object, change in
-            if LiveLink.initialized {
-                self.restartView.isHidden = !LiveLink.requiresRestart
-            }
-        }))
         
         NotificationCenter.default.addObserver(self, selector: #selector(gameControllerDidConnectNotification), name: .GCControllerDidConnect, object: nil)
         NotificationCenter.default.addObserver(self, selector: #selector(gameControllerDidDisconnectNotification), name: .GCControllerDidDisconnect, object: nil)
@@ -126,42 +121,23 @@ class StartViewController : BaseViewController {
                 gc.playerIndex = .index1
             }
         }
+        
+        LiveLink.initialize(self)
+        restartLiveLink()
     }
      
     func restartLiveLink() {
 
-        do {
-            try LiveLink.initialize(self)
-
-        } catch LiveLinkError.unknownEngineVersion {
-
-            // if there is no engine version (first run), then we ask the user to choose.
-            let askVersionAlert = UIAlertController(title: NSLocalizedString("version-title", value:"Engine Version", comment: "Title of an settings screen."),
-                                                    message: NSLocalizedString("version-choose", value:"Choose the version of Unreal Engine you will be connecting to. This selection can be modified later in the app's settings.", comment: "Message for seelecting the correct version of unreal engine."), preferredStyle: .alert)
-            askVersionAlert.addAction(UIAlertAction(title: "Unreal Engine 4", style: .default, handler: { _ in
-                LiveLink.engineVersion = .ue4
-                self.restartLiveLink()
-            }))
-            askVersionAlert.addAction(UIAlertAction(title: "Unreal Engine 5", style: .default, handler: { _ in
-                LiveLink.engineVersion = .ue5
-                self.restartLiveLink()
-            }))
-            self.present(askVersionAlert, animated: true)
-            return
-        } catch {
-            
-        }
-        
         // stop the provider & restart livelink here
         if self.liveLink != nil {
             Log.info("Restarting Messaging Engine.")
-            try? LiveLink.restart()
+            LiveLink.restart()
             self.liveLink = nil
         }
 
         Log.info("Initializing LiveLink Provider.")
 
-        self.liveLink = try? LiveLink("VCAM IOS")
+        self.liveLink = LiveLink.createProvider("Live Link VCAM")
         self.liveLink?.addCameraSubject(AppSettings.shared.liveLinkSubjectName)
 
         if let videoViewController = self.presentedViewController as? VideoViewController {
@@ -190,16 +166,12 @@ class StartViewController : BaseViewController {
         
         liveLinkTimer?.invalidate()
         liveLinkTimer = Timer.scheduledTimer(withTimeInterval: 1.0/10.0, repeats: true, block: { timer in
-            self.liveLink?.updateSubject(AppSettings.shared.liveLinkSubjectName, transform: simd_float4x4(), time: Timecode.create().toTimeInterval())
+            self.liveLink?.updateSubject(AppSettings.shared.liveLinkSubjectName, withTransform: simd_float4x4(), atTime: Timecode.create().toTimeInterval())
         })
     }
     
     override func viewDidAppear(_ animated: Bool) {
         super.viewDidAppear(animated)
-        
-        if !LiveLink.initialized {
-            restartLiveLink();
-        }
     }
     
     override func viewDidDisappear(_ animated: Bool) {
@@ -368,12 +340,9 @@ extension StartViewController : GCDAsyncUdpSocketDelegate {
     }
 }
 
-extension StartViewController : UE5LiveLinkLogDelegate, UE4LiveLinkLogDelegate {
+extension StartViewController : LiveLinkLogDelegate {
 
-    func ue5LogMessage(_ message: String!) {
-        Log.info(message)
-    }
-    func ue4LogMessage(_ message: String!) {
+    func logMessage(_ message: String!) {
         Log.info(message)
     }
 }
diff --git a/Engine/Extras/VirtualProduction/LiveLinkVCAM/vcam/VCAM-Bridging-Header.h b/Engine/Extras/VirtualProduction/LiveLinkVCAM/vcam/VCAM-Bridging-Header.h
--- a/Engine/Extras/VirtualProduction/LiveLinkVCAM/vcam/VCAM-Bridging-Header.h
+++ b/Engine/Extras/VirtualProduction/LiveLinkVCAM/vcam/VCAM-Bridging-Header.h
@@ -6,7 +6,6 @@
 
 #import "PassthroughShaderTypes.h"
 #import "ObjCUtility.h"
-#import <UE5LiveLinkProvider/UE5LiveLink.h>
-#import <UE4LiveLinkProvider/UE4LiveLink.h>
+#import <LiveLinkProvider/LiveLink.h>
 
 #endif /* VCAM_Bridging_Header_h */
diff --git a/Engine/Extras/VirtualProduction/LiveLinkVCAM/vcam/VideoViewController+ARSCNViewDelegate.swift b/Engine/Extras/VirtualProduction/LiveLinkVCAM/vcam/VideoViewController+ARSCNViewDelegate.swift
--- a/Engine/Extras/VirtualProduction/LiveLinkVCAM/vcam/VideoViewController+ARSCNViewDelegate.swift
+++ b/Engine/Extras/VirtualProduction/LiveLinkVCAM/vcam/VideoViewController+ARSCNViewDelegate.swift
@@ -13,7 +13,7 @@ extension VideoViewController : ARSCNViewDelegate {
     func renderer(_ renderer: SCNSceneRenderer, updateAtTime time: TimeInterval) {
         
         if let pov = renderer.pointOfView {
-            self.liveLink?.updateSubject(AppSettings.shared.liveLinkSubjectName, transform: pov.simdTransform, time: Timecode.create().toTimeInterval())
+            self.liveLink?.updateSubject(AppSettings.shared.liveLinkSubjectName, withTransform: pov.simdTransform, atTime: Timecode.create().toTimeInterval())
         }
         
         // update controller
diff --git a/Engine/Extras/VirtualProduction/LiveLinkVCAM/vcam/VideoViewController.swift b/Engine/Extras/VirtualProduction/LiveLinkVCAM/vcam/VideoViewController.swift
--- a/Engine/Extras/VirtualProduction/LiveLinkVCAM/vcam/VideoViewController.swift
+++ b/Engine/Extras/VirtualProduction/LiveLinkVCAM/vcam/VideoViewController.swift
@@ -36,7 +36,7 @@ class VideoViewController : BaseViewController {
 
     var oscConnection : OSCTCPConnection?
 
-    weak var liveLink : LiveLink?
+    weak var liveLink : LiveLinkProvider?
     private var decoder : VideoDecoder?
     var dismissOnDisconnect = false
     
