Added copyright boilerplate to headers.

diff --git a/Engine/Extras/VirtualProduction/LiveLinkVCAM/LiveLinkProvider.framework/Headers/LiveLink.h b/Engine/Extras/VirtualProduction/LiveLinkVCAM/LiveLinkProvider.framework/Headers/LiveLink.h
--- a/Engine/Extras/VirtualProduction/LiveLinkVCAM/LiveLinkProvider.framework/Headers/LiveLink.h
+++ b/Engine/Extras/VirtualProduction/LiveLinkVCAM/LiveLinkProvider.framework/Headers/LiveLink.h
@@ -1,3 +1,6 @@
+//
+//  Copyright Epic Games, Inc. All Rights Reserved.
+//
 
 #import <Foundation/Foundation.h>
 
diff --git a/Engine/Extras/VirtualProduction/LiveLinkVCAM/LiveLinkProvider.framework/Headers/LiveLinkProvider.h b/Engine/Extras/VirtualProduction/LiveLinkVCAM/LiveLinkProvider.framework/Headers/LiveLinkProvider.h
--- a/Engine/Extras/VirtualProduction/LiveLinkVCAM/LiveLinkProvider.framework/Headers/LiveLinkProvider.h
+++ b/Engine/Extras/VirtualProduction/LiveLinkVCAM/LiveLinkProvider.framework/Headers/LiveLinkProvider.h
@@ -1,3 +1,6 @@
+//
+//  Copyright Epic Games, Inc. All Rights Reserved.
+//
 
 #import <Foundation/Foundation.h>
 #import <simd/simd.h>
