Live Link VCAM iOS (v1.2.0 build 1) :

* Fixed a bug where the touch end events weren't properly sent during multi-touch
* Updated Tentacle SDK to add support for new Tentacle devices
* Timecodes created from the system timer now roll over to 0 hours after hour 23 (otherwise could get very high)

diff --git a/Engine/Extras/VirtualProduction/LiveLinkVCAM/vcam/Tentacle.swift b/Engine/Extras/VirtualProduction/LiveLinkVCAM/vcam/Tentacle.swift
--- a/Engine/Extras/VirtualProduction/LiveLinkVCAM/vcam/Tentacle.swift
+++ b/Engine/Extras/VirtualProduction/LiveLinkVCAM/vcam/Tentacle.swift
@@ -126,7 +126,7 @@ class Tentacle : NSObject {
         let timecode  = TentacleTimecodeAtTimestamp(device.timecode, frameRate, dropFrame, timestamp);
 
         let stringPointer = initStringPointer(capacity: TENTACLE_TIMECODE_STRING_LENGTH)
-        TentacleTimecodeString(timecode, dropFrame, stringPointer)
+        TentacleTimecodeString(timecode, frameRate, dropFrame, stringPointer)
         return String(cString: stringPointer)
     }
 }
@@ -134,7 +134,7 @@ class Tentacle : NSObject {
 
 extension Tentacle : TentacleBluetoothControllerDelegate {
 
-    func didReceiveAdvertisment(forDeviceIndex deviceIndex: Int) {
+    func didReceiveAdvertisement(forDeviceIndex deviceIndex: Int) {
 
         
         var device = TentacleDeviceCacheGetDevice(Int32(deviceIndex))
diff --git a/Engine/Extras/VirtualProduction/LiveLinkVCAM/vcam/TentacleDetailViewController.swift b/Engine/Extras/VirtualProduction/LiveLinkVCAM/vcam/TentacleDetailViewController.swift
--- a/Engine/Extras/VirtualProduction/LiveLinkVCAM/vcam/TentacleDetailViewController.swift
+++ b/Engine/Extras/VirtualProduction/LiveLinkVCAM/vcam/TentacleDetailViewController.swift
@@ -44,7 +44,7 @@ class TentacleDetailViewController : UIViewController {
         isActive = Tentacle.devicesIdentical(self.device, Tentacle.shared?.activeDevice)
         
         self.navigationItem.title = Tentacle.nameString(advertisement: &self.device.advertisement)
-        iconImageView.image = TentacleIcon.tentacleIcon(with: UInt(self.device.advertisement.icon))
+        iconImageView.image = TentacleIcon.deviceIcon(with: UInt(self.device.advertisement.icon), productId: self.device.advertisement.productId)
         uuidLabel.text = Tentacle.advertisementIdString(advertisement: self.device.advertisement)
     }
     
diff --git a/Engine/Extras/VirtualProduction/LiveLinkVCAM/vcam/Timecode.swift b/Engine/Extras/VirtualProduction/LiveLinkVCAM/vcam/Timecode.swift
--- a/Engine/Extras/VirtualProduction/LiveLinkVCAM/vcam/Timecode.swift
+++ b/Engine/Extras/VirtualProduction/LiveLinkVCAM/vcam/Timecode.swift
@@ -153,7 +153,7 @@ func clockFrequency() -> Double {
         source = timecodeSource
         valid = true
 
-        hours = Int32(timeInterval / 3600.0)
+        hours = Int32(timeInterval / 3600.0) % 24
         minutes = Int32(UInt64(timeInterval / 60.0) % 60)
         seconds = Int32(UInt64(timeInterval) % 60)
         
diff --git a/Engine/Extras/VirtualProduction/LiveLinkVCAM/vcam/TimecodeViewController.swift b/Engine/Extras/VirtualProduction/LiveLinkVCAM/vcam/TimecodeViewController.swift
--- a/Engine/Extras/VirtualProduction/LiveLinkVCAM/vcam/TimecodeViewController.swift
+++ b/Engine/Extras/VirtualProduction/LiveLinkVCAM/vcam/TimecodeViewController.swift
@@ -235,8 +235,8 @@ class TimecodeViewController : UITableViewController {
                     let tentacleCell = tableView.dequeueReusableCell(withIdentifier: "tentaclePeripheral") as! TentaclePeripheralTableViewCell
                     tentacleCell.nameLabel.text = Tentacle.nameString(advertisement: &device.advertisement)
                     tentacleCell.checkImageView.isHidden = !isSelected
-                    tentacleCell.batteryImageView.image = isCurrent ? TentacleIcon.tentacleBatteryIcon(forPercent: UInt(device.advertisement.battery)) : nil
-                    tentacleCell.signalImageView.image = isCurrent ? TentacleIcon.tentacleSignalStrengthIcon(forRssiValue: Int(device.advertisement.rssi), linkMode: device.advertisement.linkMode) : TentacleIcon.tentacleSignalStrengthIcon(with: 0)
+                    tentacleCell.batteryImageView.image = isCurrent ? TentacleIcon.batteryIcon(forPercent: UInt(device.advertisement.battery)) : nil
+                    tentacleCell.signalImageView.image = isCurrent ? TentacleIcon.signalStrengthIcon(forRssiValue: Int(device.advertisement.rssi)) : TentacleIcon.signalStrengthIcon(with: 0)
 
                     cell = tentacleCell
                 }
@@ -392,7 +392,7 @@ extension TimecodeViewController : TentacleDelegate {
                 if Tentacle.devicesIdentical(d, device) {
 
                     if d.advertisement.battery != device.advertisement.battery ||
-                        TentacleIcon.tentacleSignalStrengthIcon(forRssiValue: Int(d.advertisement.rssi), linkMode: d.advertisement.linkMode).hashValue != TentacleIcon.tentacleSignalStrengthIcon(forRssiValue: Int(device.advertisement.rssi), linkMode: device.advertisement.linkMode).hashValue ||
+                        TentacleIcon.signalStrengthIcon(forRssiValue:Int(d.advertisement.rssi)).hashValue != TentacleIcon.signalStrengthIcon(forRssiValue: Int(device.advertisement.rssi)).hashValue ||
                         Tentacle.nameString(advertisement: &d.advertisement) != Tentacle.nameString(advertisement: &device.advertisement) {
                         self.tentacleDevices[i] = device
                         self.tableView.reloadRows(at: [ IndexPath(row: i + 1, section: 0) ], with: .none)
diff --git a/Engine/Extras/VirtualProduction/LiveLinkVCAM/vcam/VideoViewController.swift b/Engine/Extras/VirtualProduction/LiveLinkVCAM/vcam/VideoViewController.swift
--- a/Engine/Extras/VirtualProduction/LiveLinkVCAM/vcam/VideoViewController.swift
+++ b/Engine/Extras/VirtualProduction/LiveLinkVCAM/vcam/VideoViewController.swift
@@ -454,6 +454,8 @@ class VideoViewController : BaseViewController {
     
     func sendOSCAddressPattern(_ pattern : OSCAddressPattern, point : CGPoint, finger : Int, force : CGFloat) {
 
+        Log.debug("\(pattern.rawValue) finger \(finger) @ \(point.x),\(point.y)")
+        
         guard let textureAspectRatio = self.pixelBufferAspectRatio else { return }
 
         // we need to adjust for the size of the top toolbar in the view : this toolbar is *over* the
@@ -536,9 +538,10 @@ class VideoViewController : BaseViewController {
     override func touchesEnded(_ touches: Set<UITouch>, with event: UIEvent?) {
 
         guard relayTouchEvents else { return }
-        
+
         for touch in touches {
             
+            
             if let fingerIndex = self.currentTouches.firstIndex(of: touch) {
                 self.sendOSCAddressPattern(OSCAddressPattern.touchEnded, point: touch.location(in: self.arView), finger: fingerIndex, force: touch.force)
                 
@@ -549,7 +552,9 @@ class VideoViewController : BaseViewController {
         
         // compress the array : remove all nils from the end
         if let lastNonNilIndex = self.currentTouches.lastIndex(where: { $0 != nil }) {
-            self.currentTouches.removeSubrange(lastNonNilIndex..<self.currentTouches.count)
+            if lastNonNilIndex < (self.currentTouches.count - 1)  {
+                self.currentTouches.removeSubrange((lastNonNilIndex + 1)..<self.currentTouches.count)
+            }
         } else {
             self.currentTouches.removeAll()
         }
