Copyright fix

diff --git a/Engine/Extras/VirtualProduction/LiveLinkVCAM/LiveLinkProvider.framework/Headers/LiveLink.h b/Engine/Extras/VirtualProduction/LiveLinkVCAM/LiveLinkProvider.framework/Headers/LiveLink.h
--- a/Engine/Extras/VirtualProduction/LiveLinkVCAM/LiveLinkProvider.framework/Headers/LiveLink.h
+++ b/Engine/Extras/VirtualProduction/LiveLinkVCAM/LiveLinkProvider.framework/Headers/LiveLink.h
@@ -1,6 +1,4 @@
-//
-//  Copyright Epic Games, Inc. All Rights Reserved.
-//
+// Copyright Epic Games, Inc. All Rights Reserved.
 
 #import <Foundation/Foundation.h>
 
diff --git a/Engine/Extras/VirtualProduction/LiveLinkVCAM/LiveLinkProvider.framework/Headers/LiveLinkProvider.h b/Engine/Extras/VirtualProduction/LiveLinkVCAM/LiveLinkProvider.framework/Headers/LiveLinkProvider.h
--- a/Engine/Extras/VirtualProduction/LiveLinkVCAM/LiveLinkProvider.framework/Headers/LiveLinkProvider.h
+++ b/Engine/Extras/VirtualProduction/LiveLinkVCAM/LiveLinkProvider.framework/Headers/LiveLinkProvider.h
@@ -1,6 +1,4 @@
-//
-//  Copyright Epic Games, Inc. All Rights Reserved.
-//
+// Copyright Epic Games, Inc. All Rights Reserved.
 
 #import <Foundation/Foundation.h>
 #import <simd/simd.h>
diff --git a/Engine/Source/Runtime/SymsLib/syms/core/extras/syms_aggregate_proc_map.h b/Engine/Source/Runtime/SymsLib/syms/core/extras/syms_aggregate_proc_map.h
--- a/Engine/Source/Runtime/SymsLib/syms/core/extras/syms_aggregate_proc_map.h
+++ b/Engine/Source/Runtime/SymsLib/syms/core/extras/syms_aggregate_proc_map.h
@@ -1,3 +1,5 @@
+// Copyright Epic Games, Inc. All Rights Reserved.
+
 /* date = January 18th 2022 1:40 pm */
 
 #ifndef SYMS_AGGREGATE_PROC_MAP_H
