Abstracted remote connection to "StreamingConnection" class with two implementations : LiveLink/RemoteSession and WebRTC.  A single streaming connection is responsible for the relay of xform/touches/gamepad updates as well as displaying the output (video) data.  The WebRTC implemation is adapted from Tensor's WebRTC-iOS example.

The new abstraction allowed us to remote some overly complicated ARSCN/Metal code that was left over from the early days.

diff --git a/Engine/Extras/VirtualProduction/LiveLinkVCAM/vcam/AppSettings.swift b/Engine/Extras/VirtualProduction/LiveLinkVCAM/vcam/AppSettings.swift
--- a/Engine/Extras/VirtualProduction/LiveLinkVCAM/vcam/AppSettings.swift
+++ b/Engine/Extras/VirtualProduction/LiveLinkVCAM/vcam/AppSettings.swift
@@ -35,6 +35,7 @@ class AppSettings : NSObject {
 
         static let showStreamingStats = "showStreamingStats"
 
+        static let connectionType = "connectionType"
     }
     
     private class func defaultsDictionary() -> Dictionary<String, Any> {
@@ -50,7 +51,9 @@ class AppSettings : NSObject {
             Keys.tentaclePeripheralName : "",
             Keys.ntpPool : "",
 
-            Keys.showStreamingStats : false
+            Keys.showStreamingStats : false,
+            
+            Keys.connectionType : "RemoteSession"
 
         ]
     }
@@ -130,4 +133,12 @@ class AppSettings : NSObject {
             UserDefaults.standard.set(showStreamingStats, forKey: Keys.showStreamingStats)
         }
     }
+    
+    @objc dynamic var connectionType : String = UserDefaults.standard.string(forKey: Keys.connectionType) ?? "RemoteSession" {
+        didSet {
+            UserDefaults.standard.set(connectionType, forKey: Keys.connectionType)
+        }
+    }
+    
+
 }
diff --git a/Engine/Extras/VirtualProduction/LiveLinkVCAM/vcam/Extensions.swift b/Engine/Extras/VirtualProduction/LiveLinkVCAM/vcam/Extensions.swift
--- a/Engine/Extras/VirtualProduction/LiveLinkVCAM/vcam/Extensions.swift
+++ b/Engine/Extras/VirtualProduction/LiveLinkVCAM/vcam/Extensions.swift
@@ -131,3 +131,15 @@ extension simd_quatf {
     }
     
 }
+
+extension Float {
+     func toBytes() -> [UInt8] {
+         withUnsafeBytes(of: self, Array.init)
+     }
+ }
+
+ extension Double {
+     func toBytes() -> [UInt8] {
+         withUnsafeBytes(of: self, Array.init)
+     }
+ }
diff --git a/Engine/Extras/VirtualProduction/LiveLinkVCAM/vcam/Info.plist b/Engine/Extras/VirtualProduction/LiveLinkVCAM/vcam/Info.plist
--- a/Engine/Extras/VirtualProduction/LiveLinkVCAM/vcam/Info.plist
+++ b/Engine/Extras/VirtualProduction/LiveLinkVCAM/vcam/Info.plist
@@ -41,6 +41,8 @@
 	<string>© 2021 Epic Games, Inc. All rights reserved.</string>
 	<key>NSLocalNetworkUsageDescription</key>
 	<string>VCAM transmits &amp; receives data from Unreal Engine.</string>
+	<key>NSMicrophoneUsageDescription</key>
+	<string>We don't use the microphone.</string>
 	<key>UIApplicationSceneManifest</key>
 	<dict>
 		<key>UIApplicationSupportsMultipleScenes</key>
@@ -77,12 +79,10 @@
 	<key>UISupportedInterfaceOrientations</key>
 	<array>
 		<string>UIInterfaceOrientationLandscapeRight</string>
-		<string>UIInterfaceOrientationLandscapeLeft</string>
 	</array>
 	<key>UISupportedInterfaceOrientations~ipad</key>
 	<array>
 		<string>UIInterfaceOrientationLandscapeRight</string>
-		<string>UIInterfaceOrientationLandscapeLeft</string>
 	</array>
 	<key>UIViewControllerBasedStatusBarAppearance</key>
 	<true/>
diff --git a/Engine/Extras/VirtualProduction/LiveLinkVCAM/vcam/JPEGVideoDecoder.swift b/Engine/Extras/VirtualProduction/LiveLinkVCAM/vcam/JPEGVideoDecoder.swift
--- a/Engine/Extras/VirtualProduction/LiveLinkVCAM/vcam/JPEGVideoDecoder.swift
+++ b/Engine/Extras/VirtualProduction/LiveLinkVCAM/vcam/JPEGVideoDecoder.swift
@@ -55,7 +55,7 @@ class CMBlockBufferHelper {
     }
 }
 
-class JPEGVideoDecoder : VideoDecoder {
+class JPEGVideoDecoder {
 
     var width : Int32 = 0
     var height : Int32 = 0
@@ -121,7 +121,8 @@ class JPEGVideoDecoder : VideoDecoder {
                 if status == 0 {
                     
                     let attrs = [kCVPixelBufferMetalCompatibilityKey : true,
-                                 kCVPixelBufferPixelFormatTypeKey : kCVPixelFormatType_32ARGB
+                                 kCVPixelBufferPixelFormatTypeKey : kCVPixelFormatType_32ARGB,
+                                 kCVPixelBufferIOSurfaceCoreAnimationCompatibilityKey : true
                         ] as [CFString : Any]
 
                     status = VTDecompressionSessionCreate(allocator: nil, formatDescription: self.formatDescription!, decoderSpecification: nil, imageBufferAttributes: attrs as CFDictionary, outputCallback : nil, decompressionSessionOut: &self.decompressor)
diff --git a/Engine/Extras/VirtualProduction/LiveLinkVCAM/vcam/PassthroughShaderTypes.h b/Engine/Extras/VirtualProduction/LiveLinkVCAM/vcam/PassthroughShaderTypes.h
deleted file mode 100644
diff --git a/Engine/Extras/VirtualProduction/LiveLinkVCAM/vcam/PassthroughShaders.metal b/Engine/Extras/VirtualProduction/LiveLinkVCAM/vcam/PassthroughShaders.metal
deleted file mode 100644
diff --git a/Engine/Extras/VirtualProduction/LiveLinkVCAM/vcam/RemoteSessionStreamingConnection.swift b/Engine/Extras/VirtualProduction/LiveLinkVCAM/vcam/RemoteSessionStreamingConnection.swift
new file mode 100644
--- /dev/null
+++ b/Engine/Extras/VirtualProduction/LiveLinkVCAM/vcam/RemoteSessionStreamingConnection.swift
@@ -0,0 +1,305 @@
+//
+//  RemoteSessionStreamingConnection.swift
+//  vcam
+//
+//  Created by Brian Smith on 6/28/22.
+//  Copyright Epic Games, Inc. All Rights Reserved.
+//
+
+import UIKit
+import CocoaAsyncSocket
+import CoreVideo
+import QuickLayout
+
+class LiveLinkLogger : LiveLinkLogDelegate {
+
+    static let instance = LiveLinkLogger()
+    
+    func logMessage(_ message: String!) {
+        Log.info(message)
+    }
+
+}
+
+
+class RemoteSessionStreamingConnection : StreamingConnection  {
+    
+    static var DefaultPort = UInt16(2049)
+
+    static var _initializedLiveLink = false
+    private var oscConnection : OSCTCPConnection?
+    private var multicastWatchdogSocket  : GCDAsyncUdpSocket?
+    private var _liveLink : LiveLinkProvider?
+    private var _timer : Timer?
+    
+    private var _decoder : JPEGVideoDecoder?
+    private var _decodedPixelBuffer : CVPixelBuffer?
+    weak var remoteSessionView : RemoteSessionView?
+    
+    private var _recvHistory = [(Date,Int)]()
+
+    var host : String?
+    var port : UInt16?
+    
+    
+    override var renderView : UIView? {
+        didSet {
+            let v = RemoteSessionView()
+            v.delegate = self
+            self.renderView?.addSubview(v)
+            remoteSessionView = v
+        }
+    }
+    
+    override var name : String {
+        get {
+            "RemoteSession"
+        }
+    }
+    
+    override var subjectName: String! {
+        didSet {
+            self._liveLink?.removeCameraSubject(oldValue!)
+            self._liveLink?.addCameraSubject(subjectName)
+        }
+    }
+    
+    override var destination : String {
+        get {
+            "\(self.host ?? "<notset>"):\(self.port ?? 0)"
+        }
+        set {
+            (self.host,self.port) = NetUtility.hostAndPortFromAddress(newValue)
+        }
+    }
+
+    override var isConnected: Bool {
+        get {
+            return self.oscConnection?.isConnected ?? false
+        }
+    }
+    
+    required init(subjectName: String) {
+        super.init(subjectName: subjectName)
+
+        // make sure LiveLink is initialized
+        if RemoteSessionStreamingConnection._initializedLiveLink == false {
+            LiveLink.initialize(LiveLinkLogger.instance)
+            RemoteSessionStreamingConnection._initializedLiveLink = true
+        }
+
+        
+        restartLiveLink()
+    }
+
+    override func connect() throws {
+        self.oscConnection = nil
+        self.oscConnection = try OSCTCPConnection(host:self.host ?? "", port:self.port ?? RemoteSessionStreamingConnection.DefaultPort, delegate: self)
+        
+        _timer = Timer.scheduledTimer(withTimeInterval: 1, repeats: true, block: { (t) in
+            
+            if let conn = self.oscConnection {
+                if conn.isConnected {
+                    conn.send(OSCAddressPattern.ping)
+                } else {
+                    conn.reconnect()
+                }
+            }
+            
+            let now = Date()
+            var countToRemove = 0
+            var totalBytesLastSecond = 0
+            for item in self._recvHistory {
+                if now.timeIntervalSince(item.0) > 1 {
+                    countToRemove += 1
+                } else {
+                    totalBytesLastSecond += item.1
+                }
+            }
+
+            self._recvHistory.removeFirst(countToRemove)
+
+            self.stats = StreamingConnectionStats()
+            self.stats?.bytesPerSecond = totalBytesLastSecond
+
+            if self._recvHistory.count > 1,
+               let first = self._recvHistory.first,
+               let last = self._recvHistory.last {
+                
+                self.stats?.framesPerSecond = Float(self._recvHistory.count - 1) / Float(last.0.timeIntervalSince(first.0))
+            }
+
+        })
+    }
+    
+    override func reconnect() {
+        
+        oscConnection?.reconnect()
+    }
+    
+    override func disconnect() {
+        
+        _timer?.invalidate()
+        _timer = nil
+        
+        oscConnection?.disconnect()
+        _decoder = nil
+    }
+    
+    override func sendTransform(_ transform: simd_float4x4, atTime time: Double) {
+        self._liveLink?.updateSubject(AppSettings.shared.liveLinkSubjectName, withTransform: transform, atTime: time)
+    }
+    
+    func restartLiveLink() {
+
+        // stop the provider & restart livelink here
+        if self._liveLink != nil {
+            Log.info("Restarting LiveLink.")
+            LiveLink.restart()
+            self._liveLink = nil
+        }
+
+        Log.info("Initializing LiveLink Provider.")
+
+        self._liveLink = LiveLink.createProvider("Live Link VCAM")
+        self._liveLink?.addCameraSubject(AppSettings.shared.liveLinkSubjectName)
+
+        multicastWatchdogSocket?.close()
+        Log.info("Starting multicast watchdog.")
+        multicastWatchdogSocket = GCDAsyncUdpSocket(delegate: self, delegateQueue: DispatchQueue.main)
+        do {
+            try multicastWatchdogSocket?.enableReusePort(true)
+            try multicastWatchdogSocket?.bind(toPort: 6665)
+            try multicastWatchdogSocket?.joinMulticastGroup("230.0.0.1")
+            try multicastWatchdogSocket?.beginReceiving()
+        } catch {
+            Log.info("Error creating watchdog : \(error.localizedDescription)")
+        }
+    }
+    
+    func decode(_ width : Int32, _ height : Int32, _ blob : Data) {
+
+        if _decoder == nil {
+            _decoder = JPEGVideoDecoder()
+        }
+
+        _recvHistory.append((Date(), blob.count))
+        
+        self._decoder?.decode(width: width, height: height, data: blob) { (pixelBuffer) in
+            
+            if let pb = pixelBuffer {
+                
+                DispatchQueue.main.async {
+                    self.remoteSessionView?.pixelBuffer = pb
+                    self.videoSize = CGSize(width: CVPixelBufferGetWidth(pb), height: CVPixelBufferGetHeight(pb))
+                }
+            }
+            
+        }
+    }
+}
+
+extension RemoteSessionStreamingConnection : RemoteSessionViewDelegate {
+    
+    func remoteSessionView(_ view: RemoteSessionView?, touch type: StreamingConnectionTouchType, index: Int, at point: CGPoint, force: CGFloat) {
+        
+        guard let rsv = remoteSessionView else { return }
+        
+        let normalizedPoint = CGPoint(x: point.x / rsv.frame.size.width, y: point.y / rsv.frame.size.height)
+        let data = OSCUtility.ueTouchData(point: normalizedPoint, finger: index, force: force)
+        
+        //Log.info("\(type.rawValue) @ \(normalizedPoint)")
+        
+        var pattern : OSCAddressPattern!
+        switch type {
+        case .began:
+            pattern = .touchStarted
+        case .moved:
+            pattern = .touchMoved
+        case .ended:
+            pattern = .touchEnded
+        }
+        
+        self.oscConnection?.send(pattern, arguments: [ OSCArgument.blob(data) ])
+
+    }
+}
+
+
+extension RemoteSessionStreamingConnection : GCDAsyncUdpSocketDelegate {
+
+    func udpSocketDidClose(_ sock: GCDAsyncUdpSocket, withError error: Error?) {
+        Log.error("Multicast watchdog closed : restarting LiveLink.")
+        self.restartLiveLink()
+    }
+}
+
+extension RemoteSessionStreamingConnection : OSCTCPConnectionDelegate {
+    
+    func oscConnectionDidConnect(_ connection: OSCTCPConnection) {
+        Log.info("RemoteSession connection initiated.")
+
+    }
+    
+    func oscConnectionDidDisconnect(_ connection: OSCTCPConnection, withError err: Error?) {
+        self.delegate?.streamingConnection(self, didDisconnectWithError: err)
+    }
+    
+    func oscConnection(_ connection: OSCTCPConnection, didReceivePacket packet: OSCPacket) {
+        
+        if let msg = packet as? OSCPacketMessage {
+            
+            //Log.info(msg.debugString())
+            
+            switch msg.addressPattern {
+
+            case OSCAddressPattern.rsHello.rawValue:
+                if let args = msg.arguments {
+                    if args.count == 1,
+                        case let OSCArgument.string(version) = args[0] {
+                        connection.send(OSCAddressPattern.rsHello, arguments: [ OSCArgument.string(version) ])
+                    }
+                }
+
+            case OSCAddressPattern.rsGoodbye.rawValue:
+                connection.disconnect()
+                if let args = msg.arguments {
+                    if args.count == 1,
+                        case let OSCArgument.string(version) = args[0] {
+                        Log.info("RemoteSession closed by server : \(version)")
+                    }
+                }
+                
+            case OSCAddressPattern.rsChannelList.rawValue:
+                connection.send(OSCAddressPattern.rsChangeChannel, arguments: [ OSCArgument.string("FRemoteSessionImageChannel"), OSCArgument.string("Write"), OSCArgument.int32(1)] )
+                connection.send(OSCAddressPattern.rsChangeChannel, arguments: [ OSCArgument.string("FRemoteSessionInputChannel"), OSCArgument.string("Read"), OSCArgument.int32(1)] )
+                
+                self.delegate?.streamingConnectionDidConnect(self)
+
+            case OSCAddressPattern.screen.rawValue:
+
+                if let args = msg.arguments {
+                    if args.count == 4,
+                        case let OSCArgument.int32(width) = args[0],
+                        case let OSCArgument.int32(height) = args[1],
+                       case let OSCArgument.blob(blob) = args[2] {
+                        
+                        //Log.info("RECV frame \(frame) : \(width)x\(height) \(blob.count) bytes")
+                        self.decode(width, height, blob)
+                    }
+                }
+                
+            case OSCAddressPattern.ping.rawValue:
+                connection.send(OSCAddressPattern.ping)
+
+
+            default:
+                break
+            }
+        }
+    }
+    
+    
+}
+
+
diff --git a/Engine/Extras/VirtualProduction/LiveLinkVCAM/vcam/RemoteSessionView.swift b/Engine/Extras/VirtualProduction/LiveLinkVCAM/vcam/RemoteSessionView.swift
new file mode 100644
--- /dev/null
+++ b/Engine/Extras/VirtualProduction/LiveLinkVCAM/vcam/RemoteSessionView.swift
@@ -0,0 +1,120 @@
+//
+//  RemoteSessionView.swift
+//  vcam
+//
+//  Created by Brian Smith on 6/30/22.
+//  Copyright Epic Games, Inc. All Rights Reserved.
+//
+
+import UIKit
+import QuickLayout
+
+protocol RemoteSessionViewDelegate {
+    
+    func remoteSessionView(_ view : RemoteSessionView?, touch : StreamingConnectionTouchType, index : Int, at point: CGPoint, force: CGFloat)
+    
+}
+
+class RemoteSessionView : UIView {
+
+    var delegate : RemoteSessionViewDelegate?
+    
+    private var _currentTouches = Array<UITouch?>()
+
+    private weak var _aspectRatioConstraint : NSLayoutConstraint?
+
+    var pixelBuffer : CVPixelBuffer? {
+        didSet {
+            if let pb = pixelBuffer {
+
+                self.layer.contents = pb
+                
+                let thisSize = CGSize(width: CVPixelBufferGetWidth(pb), height: CVPixelBufferGetHeight(pb))
+                if _aspectRatioConstraint == nil || _aspectRatioConstraint?.multiplier != thisSize.width / thisSize.height {
+                    if let constraint = _aspectRatioConstraint {
+                        self.removeConstraint(constraint)
+                    }
+                    _aspectRatioConstraint = layout(.width, to: .height, of: self, ratio: (thisSize.width / thisSize.height), priority: .required)
+                }
+            }
+        }
+    }
+    
+    override func didMoveToSuperview() {
+        self.removeConstraints(self.constraints)
+        
+        // explained here :
+        // https://stackoverflow.com/questions/25766747/emulating-aspect-fit-behaviour-using-autolayout-constraints-in-xcode-6
+        layoutToSuperview(.centerX, .centerY)
+        layoutToSuperview(.width, relation: .lessThanOrEqual, priority: .required)
+        layoutToSuperview(.width, relation: .equal, priority: .defaultHigh)
+        layoutToSuperview(.height, relation: .lessThanOrEqual, priority: .required)
+        layoutToSuperview(.height, relation: .equal, priority: .defaultHigh)
+    }
+    
+    override func touchesBegan(_ touches: Set<UITouch>, with event: UIEvent?) {
+
+        //guard relayTouchEvents else { return }
+        
+        // for each touch, we need to maintain an index to send via OSC, so that UE
+        // can correctly map between began/moved. The currentTouches array will maintain
+        // the indices for each touch (keeping some nil if needed) and clean up / compress
+        // the array when a touch ends.
+        
+        for touch in touches {
+
+            // insert this new touch into the array, using the first available slot (a nil entry)
+            // or append to the end.
+            var fingerIndex = 0
+            
+            for i in 0..._currentTouches.count {
+                fingerIndex = i
+                if i == _currentTouches.count {
+                    _currentTouches.append(touch)
+                    break
+                } else if _currentTouches[i] == nil {
+                    _currentTouches[i] = touch
+                    break
+                }
+            }
+            self.delegate?.remoteSessionView(self, touch: .began, index: fingerIndex, at: touch.location(in: self), force: touch.force)
+        }
+    }
+    
+    override func touchesMoved(_ touches: Set<UITouch>, with event: UIEvent?) {
+
+        //guard relayTouchEvents else { return }
+
+        for touch in touches {
+            if let fingerIndex = _currentTouches.firstIndex(of: touch) {
+                self.delegate?.remoteSessionView(self, touch: .moved, index: fingerIndex, at: touch.location(in: self), force: touch.force)
+            }
+        }
+
+    }
+    
+    override func touchesEnded(_ touches: Set<UITouch>, with event: UIEvent?) {
+
+        //guard relayTouchEvents else { return }
+
+        for touch in touches {
+            
+            
+            if let fingerIndex = _currentTouches.firstIndex(of: touch) {
+                self.delegate?.remoteSessionView(self, touch: .ended, index: fingerIndex, at: touch.location(in: self), force: touch.force)
+                
+                // this touch is ended, set to nil
+                _currentTouches[fingerIndex] = nil
+            }
+        }
+        
+        // compress the array : remove all nils from the end
+        if let lastNonNilIndex = _currentTouches.lastIndex(where: { $0 != nil }) {
+            if lastNonNilIndex < (_currentTouches.count - 1)  {
+                _currentTouches.removeSubrange((lastNonNilIndex + 1)..<_currentTouches.count)
+            }
+        } else {
+            _currentTouches.removeAll()
+        }
+    }
+}
diff --git a/Engine/Extras/VirtualProduction/LiveLinkVCAM/vcam/SettingsViewController.swift b/Engine/Extras/VirtualProduction/LiveLinkVCAM/vcam/SettingsViewController.swift
--- a/Engine/Extras/VirtualProduction/LiveLinkVCAM/vcam/SettingsViewController.swift
+++ b/Engine/Extras/VirtualProduction/LiveLinkVCAM/vcam/SettingsViewController.swift
@@ -38,6 +38,16 @@ class SettingsViewController : UITableViewController {
 
         case "timecode":
             cell.detailTextLabel?.text = Timecode.sourceToString(appSettings.timecodeSourceEnum())
+            
+        case "connectionType":
+            switch appSettings.connectionType {
+            case "RemoteSession":
+                cell.detailTextLabel?.text = "Remote Session"
+            case "WebRTC":
+                cell.detailTextLabel?.text = "Pixel Streaming"
+            default:
+                cell.detailTextLabel?.text = "Unknown"
+            }
 
         default:
             break
@@ -47,7 +57,7 @@ class SettingsViewController : UITableViewController {
     @objc func done(sender:Any?) {
         self.navigationController?.dismiss(animated: true, completion: nil)
     }
-
+    
     override func prepare(for segue: UIStoryboardSegue, sender: Any?) {
 
         if let vc = segue.destination as? SingleValueViewController {
@@ -71,5 +81,22 @@ class SettingsViewController : UITableViewController {
                 }
             }
         }
+        
+        else if let vc = segue.destination as? MultipleChoiceViewController {
+            
+            if segue.identifier == "connectionType" {
+                
+                vc.navigationItem.title = "Connection Type"
+                vc.items = [ "Remote Session", "Pixel Streaming" ]
+                vc.selectedIndex = appSettings.connectionType == "RemoteSession" ? 0 : 1
+                vc.footerString = "Description of Connection Types"
+                vc.completion = { (index) in
+                    self.appSettings.connectionType = index == 0 ? "RemoteSession" : "WebRTC"
+                }
+            }
+        }
+
     }
+    
+    
 }
diff --git a/Engine/Extras/VirtualProduction/LiveLinkVCAM/vcam/StartViewController+OSCTCPConnectionDelegate.swift b/Engine/Extras/VirtualProduction/LiveLinkVCAM/vcam/StartViewController+OSCTCPConnectionDelegate.swift
deleted file mode 100644
diff --git a/Engine/Extras/VirtualProduction/LiveLinkVCAM/vcam/StartViewController+StreamingConnectionDelegate.swift b/Engine/Extras/VirtualProduction/LiveLinkVCAM/vcam/StartViewController+StreamingConnectionDelegate.swift
new file mode 100644
--- /dev/null
+++ b/Engine/Extras/VirtualProduction/LiveLinkVCAM/vcam/StartViewController+StreamingConnectionDelegate.swift
@@ -0,0 +1,43 @@
+//
+//  StartViewController+StreamingConnectionDelegate.swift
+//  vcam
+//
+//  Created by Brian Smith on 6/29/22.
+//  Copyright Epic Games, Inc. All Rights Reserved.
+//
+
+import UIKit
+
+extension StartViewController : StreamingConnectionDelegate {
+    
+    func streamingConnectionDidConnect(_ connection: StreamingConnection) {
+        Log.info("StreamingConnection \(connection.name) did connect to \(connection.destination)")
+        
+        DispatchQueue.main.async {
+            self.hideConnectingAlertView {
+                self.performSegue(withIdentifier: "showVideoView", sender: self)
+           }
+        }
+    }
+
+    func streamingConnection(_ connection: StreamingConnection, didDisconnectWithError err: Error?) {
+
+        DispatchQueue.main.async {
+
+            self.hideConnectingAlertView() {
+                if let e = err {
+                    Log.info("StreamingConnection \(connection.name) disconnected with error : \(e.localizedDescription)")
+
+                    let errorAlert = UIAlertController(title: "Error", message: "Couldn't connect : \(e.localizedDescription)", preferredStyle: .alert)
+                    errorAlert.addAction(UIAlertAction(title: "OK", style: .default, handler: { _ in
+                        self.hideConnectingView() {}
+                    }))
+                    self.present(errorAlert, animated:true)
+                    
+                } else {
+                    Log.info("StreamingConnection \(connection.name) disconnected.")
+                }
+            }
+        }
+    }
+}
diff --git a/Engine/Extras/VirtualProduction/LiveLinkVCAM/vcam/StartViewController+UELogDelegate.swift b/Engine/Extras/VirtualProduction/LiveLinkVCAM/vcam/StartViewController+UELogDelegate.swift
deleted file mode 100644
diff --git a/Engine/Extras/VirtualProduction/LiveLinkVCAM/vcam/StartViewController.swift b/Engine/Extras/VirtualProduction/LiveLinkVCAM/vcam/StartViewController.swift
--- a/Engine/Extras/VirtualProduction/LiveLinkVCAM/vcam/StartViewController.swift
+++ b/Engine/Extras/VirtualProduction/LiveLinkVCAM/vcam/StartViewController.swift
@@ -12,8 +12,6 @@ import Kronos
 import GameController
 
 class StartViewController : BaseViewController {
-    
-    static var DefaultPort = UInt16(2049)
 
     @IBOutlet weak var headerView : HeaderView!
     @IBOutlet weak var versionLabel : UILabel!
@@ -29,12 +27,10 @@ class StartViewController : BaseViewController {
     @IBOutlet weak var connectingView : UIVisualEffectView!
 
     private var tapGesture : UITapGestureRecognizer!
-    private var oscConnection : OSCTCPConnection?
-    
-    public var multicastWatchdogSocket  : GCDAsyncUdpSocket?
     
-    private var liveLink : LiveLinkProvider?
-    private var liveLinkTimer : Timer?
+    private var streamingConnection : StreamingConnection?
+    private var _liveLinkTimer : Timer?
+
     
     @objc dynamic let appSettings = AppSettings.shared
     private var observers = [NSKeyValueObservation]()
@@ -108,10 +104,24 @@ class StartViewController : BaseViewController {
         
         // any change to the subject name will remove & re-add the camera subject.
         observers.append(observe(\.appSettings.liveLinkSubjectName, options: [.old,.new], changeHandler: { object, change in
-            self.liveLink?.removeCameraSubject(change.oldValue!)
-            self.liveLink?.addCameraSubject(self.appSettings.liveLinkSubjectName)
+            if let sc = self.streamingConnection {
+                sc.subjectName = self.appSettings.liveLinkSubjectName
+            }
         }))
-        
+
+        // initial & value changes for the connection type instantiates a new StreamingConnection object
+        observers.append(observe(\.appSettings.connectionType, options: [.initial, .old,.new], changeHandler: { object, change in
+            
+            self.streamingConnection = nil
+
+            let connectionType = self.appSettings.connectionType
+            if let connectionClass = Bundle.main.classNamed("VCAM.\(connectionType)StreamingConnection") as? StreamingConnection.Type {
+                self.streamingConnection = connectionClass.init(subjectName: self.appSettings.liveLinkSubjectName)
+                self.streamingConnection?.delegate = self
+            }
+
+        }))
+
         NotificationCenter.default.addObserver(self, selector: #selector(gameControllerDidConnectNotification), name: .GCControllerDidConnect, object: nil)
         NotificationCenter.default.addObserver(self, selector: #selector(gameControllerDidDisconnectNotification), name: .GCControllerDidDisconnect, object: nil)
         
@@ -121,52 +131,19 @@ class StartViewController : BaseViewController {
                 gc.playerIndex = .index1
             }
         }
-        
-        LiveLink.initialize(self)
-        restartLiveLink()
     }
-     
-    func restartLiveLink() {
-
-        // stop the provider & restart livelink here
-        if self.liveLink != nil {
-            Log.info("Restarting Messaging Engine.")
-            LiveLink.restart()
-            self.liveLink = nil
-        }
-
-        Log.info("Initializing LiveLink Provider.")
-
-        self.liveLink = LiveLink.createProvider("Live Link VCAM")
-        self.liveLink?.addCameraSubject(AppSettings.shared.liveLinkSubjectName)
-
-        if let videoViewController = self.presentedViewController as? VideoViewController {
-            if !self.ipAddressIsDemoMode {
-                videoViewController.liveLink = self.liveLink
-            }
-        }
+    
 
-        multicastWatchdogSocket?.close()
-        Log.info("Starting multicast watchdog.")
-        multicastWatchdogSocket = GCDAsyncUdpSocket(delegate: self, delegateQueue: DispatchQueue.main)
-        do {
-            try multicastWatchdogSocket?.enableReusePort(true)
-            try multicastWatchdogSocket?.bind(toPort: 6665)
-            try multicastWatchdogSocket?.joinMulticastGroup("230.0.0.1")
-            try multicastWatchdogSocket?.beginReceiving()
-        } catch {
-            Log.info("Error creating watchdog : \(error.localizedDescription)")
-        }
-    }
     
     override func viewWillAppear(_ animated : Bool) {
         
         self.connectingView.isHidden = true
         self.headerView.start()
+
+        self.streamingConnection?.delegate = self
         
-        liveLinkTimer?.invalidate()
-        liveLinkTimer = Timer.scheduledTimer(withTimeInterval: 1.0/10.0, repeats: true, block: { timer in
-            self.liveLink?.updateSubject(AppSettings.shared.liveLinkSubjectName, withTransform: simd_float4x4(), atTime: Timecode.create().toTimeInterval())
+        _liveLinkTimer = Timer.scheduledTimer(withTimeInterval: 1.0/10.0, repeats: true, block: { timer in
+            self.streamingConnection?.sendTransform(simd_float4x4(), atTime: Timecode.create().toTimeInterval())
         })
     }
     
@@ -189,16 +166,13 @@ class StartViewController : BaseViewController {
             if let vc = segue.destination as? VideoViewController {
                 
                 // stop the timer locally which is sending LL identity xform
-                liveLinkTimer?.invalidate()
-                liveLinkTimer = nil
+                _liveLinkTimer?.invalidate()
+                _liveLinkTimer = nil
 
-                self.oscConnection?.delegate = vc
-                vc.oscConnection = self.oscConnection
+                self.streamingConnection?.delegate = vc
+                vc.streamingConnection = self.streamingConnection
                 
-                vc.liveLink = self.liveLink
                 vc.gameController = self.gameController
-
-                self.oscConnection = nil
             }
         }
     }
@@ -250,6 +224,7 @@ class StartViewController : BaseViewController {
 
         } else {
 
+            // show the connection view
             self.connectingView.isHidden = false
             self.connectingView.alpha = 0.0
             UIView.animate(withDuration: 0.2) {
@@ -257,15 +232,12 @@ class StartViewController : BaseViewController {
             }
             
             showConnectingAlertView(mode: .connecting) {
-                self.oscConnection = nil
                 self.hideConnectingView() { }
             }
 
-            let (host,port) = NetUtility.hostAndPortFromAddress(self.ipAddress.text!)
-            
             do {
-                self.oscConnection = nil
-                self.oscConnection = try OSCTCPConnection(host:host, port:port ?? StartViewController.DefaultPort, delegate: self)
+                self.streamingConnection?.destination = self.ipAddress.text!
+                try self.streamingConnection?.connect()
                 
             } catch {
                 
@@ -332,17 +304,3 @@ extension StartViewController : UITextFieldDelegate {
     }
 }
 
-extension StartViewController : GCDAsyncUdpSocketDelegate {
-
-    func udpSocketDidClose(_ sock: GCDAsyncUdpSocket, withError error: Error?) {
-        Log.error("Multicast watchdog closed : restarting LiveLink.")
-        restartLiveLink()
-    }
-}
-
-extension StartViewController : LiveLinkLogDelegate {
-
-    func logMessage(_ message: String!) {
-        Log.info(message)
-    }
-}
diff --git a/Engine/Extras/VirtualProduction/LiveLinkVCAM/vcam/StreamingConnection.swift b/Engine/Extras/VirtualProduction/LiveLinkVCAM/vcam/StreamingConnection.swift
new file mode 100644
--- /dev/null
+++ b/Engine/Extras/VirtualProduction/LiveLinkVCAM/vcam/StreamingConnection.swift
@@ -0,0 +1,102 @@
+//
+//  StreamingConnection.swift
+//  vcam
+//
+//  Created by Brian Smith on 6/28/22.
+//  Copyright Epic Games, Inc. All rights reserved.
+//
+
+import Foundation
+import UIKit
+
+protocol StreamingConnectionDelegate {
+    
+    func streamingConnectionDidConnect(_ connection : StreamingConnection)
+    func streamingConnection(_ connection : StreamingConnection, didDisconnectWithError err: Error?)
+}
+
+enum StreamingConnectionTouchType : String {
+    case began = "Began"
+    case moved = "Moved"
+    case ended = "Ended"
+}
+
+class StreamingConnectionStats {
+    var bytesPerSecond : Int?
+    var framesPerSecond : Float?
+}
+
+class StreamingConnection : NSObject {
+
+    var delegate: StreamingConnectionDelegate?
+    var subjectName : String!
+    
+    var renderView : UIView?
+
+    var name: String {
+        get {
+            assertionFailure("not implemented")
+            return ""
+        }
+    }
+    
+    var destination: String {
+        get {
+            assertionFailure("not implemented")
+            return ""
+        }
+        set {
+            assertionFailure("not implemented")
+        }
+        
+    }
+    
+    var isConnected: Bool {
+        get {
+            assertionFailure("not implemented")
+            return false
+        }
+    }
+    
+    var stats : StreamingConnectionStats?
+    
+    private(set) var videoAspectRatio: CGFloat?
+    var videoSize: CGSize? {
+        didSet {
+            if let vs = videoSize {
+                self.videoAspectRatio = vs.width / vs.height
+            } else {
+                self.videoAspectRatio = nil
+            }
+        }
+    }
+    
+    required init(subjectName: String) {
+        self.subjectName = subjectName
+    }
+
+    func connect() throws {
+        assertionFailure("not implemented")
+    }
+
+    func reconnect() {
+        assertionFailure("not implemented")
+    }
+
+    func disconnect() {
+        assertionFailure("not implemented")
+    }
+    
+    func sendTransform(_ transform: simd_float4x4, atTime time: Double) {
+        assertionFailure("not implemented")
+    }
+    
+    func setRenderView(_ view : UIView) {
+        assertionFailure("not implemented")
+    }
+    
+    func getStats() -> StreamingConnectionStats? {
+        return nil
+    }
+
+}
diff --git a/Engine/Extras/VirtualProduction/LiveLinkVCAM/vcam/VCAM-Bridging-Header.h b/Engine/Extras/VirtualProduction/LiveLinkVCAM/vcam/VCAM-Bridging-Header.h
--- a/Engine/Extras/VirtualProduction/LiveLinkVCAM/vcam/VCAM-Bridging-Header.h
+++ b/Engine/Extras/VirtualProduction/LiveLinkVCAM/vcam/VCAM-Bridging-Header.h
@@ -4,7 +4,6 @@
 #ifndef VCAM_Bridging_Header_h
 #define VCAM_Bridging_Header_h
 
-#import "PassthroughShaderTypes.h"
 #import "ObjCUtility.h"
 #import <LiveLinkProvider/LiveLink.h>
 
diff --git a/Engine/Extras/VirtualProduction/LiveLinkVCAM/vcam/VideoViewController+ARCoachingOverlayViewDelegate.swift b/Engine/Extras/VirtualProduction/LiveLinkVCAM/vcam/VideoViewController+ARCoachingOverlayViewDelegate.swift
--- a/Engine/Extras/VirtualProduction/LiveLinkVCAM/vcam/VideoViewController+ARCoachingOverlayViewDelegate.swift
+++ b/Engine/Extras/VirtualProduction/LiveLinkVCAM/vcam/VideoViewController+ARCoachingOverlayViewDelegate.swift
@@ -3,7 +3,7 @@
 //  vcam
 //
 //  Created by Brian Smith on 12/6/21.
-//  Copyright © 2021 Brian Smith. All rights reserved.
+//  Copyright Epic Games, Inc. All Rights Reserved.
 //
 
 import Foundation
diff --git a/Engine/Extras/VirtualProduction/LiveLinkVCAM/vcam/VideoViewController+ARSCNViewDelegate.swift b/Engine/Extras/VirtualProduction/LiveLinkVCAM/vcam/VideoViewController+ARSCNViewDelegate.swift
deleted file mode 100644
diff --git a/Engine/Extras/VirtualProduction/LiveLinkVCAM/vcam/VideoViewController+ARSessionDelegate.swift b/Engine/Extras/VirtualProduction/LiveLinkVCAM/vcam/VideoViewController+ARSessionDelegate.swift
--- a/Engine/Extras/VirtualProduction/LiveLinkVCAM/vcam/VideoViewController+ARSessionDelegate.swift
+++ b/Engine/Extras/VirtualProduction/LiveLinkVCAM/vcam/VideoViewController+ARSessionDelegate.swift
@@ -36,4 +36,11 @@ extension VideoViewController : ARSessionDelegate {
             
         }
     }
+    
+    func session(_ session: ARSession, didUpdate frame: ARFrame) {
+        self.streamingConnection?.sendTransform(frame.camera.transform, atTime: Timecode.create().toTimeInterval())
+        
+        // update controller
+        self.sendControllerUpdate()
+    }
 }
diff --git a/Engine/Extras/VirtualProduction/LiveLinkVCAM/vcam/VideoViewController+OSCTCPConnectionDelegate.swift b/Engine/Extras/VirtualProduction/LiveLinkVCAM/vcam/VideoViewController+OSCTCPConnectionDelegate.swift
deleted file mode 100644
diff --git a/Engine/Extras/VirtualProduction/LiveLinkVCAM/vcam/VideoViewController+StreamingConnectionDelegate.swift b/Engine/Extras/VirtualProduction/LiveLinkVCAM/vcam/VideoViewController+StreamingConnectionDelegate.swift
new file mode 100644
--- /dev/null
+++ b/Engine/Extras/VirtualProduction/LiveLinkVCAM/vcam/VideoViewController+StreamingConnectionDelegate.swift
@@ -0,0 +1,40 @@
+//
+//  VideoViewController+StreamingConnectionDelegate.swift
+//  vcam
+//
+//  Created by Brian Smith on 6/28/22.
+//  Copyright Epic Games, Inc. All Rights Reserved.
+//
+
+import Foundation
+
+extension VideoViewController : StreamingConnectionDelegate {
+
+    func streamingConnectionDidConnect(_ connection: StreamingConnection) {
+        Log.info("StreamingConnection \(connection.name) did connect to \(connection.destination)")
+        DispatchQueue.main.async {
+            self.showReconnecting(false, animated: true)
+        }
+    }
+    
+    func streamingConnection(_ connection: StreamingConnection, didDisconnectWithError err: Error?) {
+        if let e = err {
+            Log.info("StreamingConnection \(connection.name) disconnected: \(e.localizedDescription)")
+        } else {
+            Log.info("StreamingConnection \(connection.name) disconnected.")
+        }
+
+        DispatchQueue.main.async {
+
+            if self.dismissOnDisconnect {
+                self.streamingConnection?.delegate = nil
+                self.streamingConnection = nil
+                self.presentingViewController?.dismiss(animated: true, completion: nil)
+            } else {
+                self.showReconnecting(true, animated: true)
+            }
+        }
+
+    }
+}
+
diff --git a/Engine/Extras/VirtualProduction/LiveLinkVCAM/vcam/VideoViewController.swift b/Engine/Extras/VirtualProduction/LiveLinkVCAM/vcam/VideoViewController.swift
--- a/Engine/Extras/VirtualProduction/LiveLinkVCAM/vcam/VideoViewController.swift
+++ b/Engine/Extras/VirtualProduction/LiveLinkVCAM/vcam/VideoViewController.swift
@@ -14,18 +14,16 @@ import GameController
 
 class VideoViewController : BaseViewController {
 
-    enum VideoDecoderType {
-        
-        case jpeg
-        case h264
-    }
-    
     var demoMode : Bool {
         liveLink == nil
     }
 
     private var displayLink: CADisplayLink?
 
+    private var arSession : ARSession?
+    
+    @IBOutlet weak var renderView : UIView!
+    
     @IBOutlet weak var headerView : HeaderView!
     @IBOutlet weak var headerViewTopConstraint : NSLayoutConstraint!
     var headerViewY : CGFloat = 0
@@ -33,14 +31,12 @@ class VideoViewController : BaseViewController {
     var headerViewTopConstraintStartValue : CGFloat = 0
     var headerPanGestureRecognizer : UIPanGestureRecognizer!
     var headerPullDownGestureRecognizer : UIScreenEdgePanGestureRecognizer!
-
-    var oscConnection : OSCTCPConnection?
+    
+    var streamingConnection : StreamingConnection?
 
     weak var liveLink : LiveLinkProvider?
-    private var decoder : VideoDecoder?
     var dismissOnDisconnect = false
     
-    var stats = Array<(Date,Int)>()
     var statsTimer : Timer?
     
     var gameControllerSnapshot : GCController?
@@ -69,23 +65,23 @@ class VideoViewController : BaseViewController {
         }
     }
     
-    private var currentTouches = Array<UITouch?>()
-    
     var relayTouchEvents = true {
         didSet {
 
             // if we are stopping touch events sent to OSC, then we need to send an "ended" event
             // for all of the current touches so the UI doesn't think it's still occurring
+            /*
             if relayTouchEvents == false {
 
                 for i in 0..<self.currentTouches.count {
                     if let t = self.currentTouches[i] {
-                        self.sendOSCAddressPattern(OSCAddressPattern.touchEnded, point: t.location(in: self.arView), finger: i, force: 0);
+                        self.sendTouch(.ended, point: t.location(in: self.renderView), finger: i, force: 0);
                     }
                 }
                 self.currentTouches.removeAll()
 
             }
+             */
             
         }
     }
@@ -94,61 +90,7 @@ class VideoViewController : BaseViewController {
 
     @IBOutlet weak var reconnectingBlurView : UIVisualEffectView!
 
-    @IBOutlet weak var demoModeBlurView : UIVisualEffectView!
-    
-    var mtlTextureCache : CVMetalTextureCache?
-    private var mtlCommandQueue : MTLCommandQueue?
-    var mtlPassthroughPipelineState : MTLRenderPipelineState?
-    var mtlVertexBuffer : MTLBuffer?
-    var mtlNoSignalTexture : MTLTexture?
-    var mtlDecodedTexture : MTLTexture?
-
-    var pixelBufferSize : CGSize? {
-        didSet {
-            if let sz = pixelBufferSize {
-                pixelBufferAspectRatio = Float(sz.width / sz.height)
-            }
-        }
-    }
-    var pixelBufferAspectRatio : Float?
-    var _decodedPixelBufferMutex : pthread_mutex_t!
-    var _decodedPixelBuffer : CVMetalTexture?
-    var decodedPixelBuffer : CVPixelBuffer? {
-        get {
-            pthread_mutex_lock(&_decodedPixelBufferMutex)
-            defer { pthread_mutex_unlock(&_decodedPixelBufferMutex) }
-            return _decodedPixelBuffer
-        }
-        set {
-            pthread_mutex_lock(&_decodedPixelBufferMutex)
-            defer { pthread_mutex_unlock(&_decodedPixelBufferMutex) }
-            _decodedPixelBuffer = newValue
-        }
-    }
-    func takeDecodedPixelBuffer() -> CVPixelBuffer? {
-        pthread_mutex_lock(&_decodedPixelBufferMutex)
-        defer { pthread_mutex_unlock(&_decodedPixelBufferMutex) }
-        let pb = _decodedPixelBuffer
-        _decodedPixelBuffer = nil
-        return pb
-    }
-
-
-    var pingTimer : Timer?
     
-    required init?(coder: NSCoder) {
-        super.init(coder: coder)
-        
-        var attr: pthread_mutexattr_t = pthread_mutexattr_t()
-        pthread_mutexattr_init(&attr)
-        _decodedPixelBufferMutex = pthread_mutex_t()
-        pthread_mutex_init(&_decodedPixelBufferMutex, &attr);
-    }
-    
-    deinit {
-        pthread_mutex_destroy(&_decodedPixelBufferMutex)
-    }
-
     override var prefersHomeIndicatorAutoHidden: Bool {
         return true
     }
@@ -186,13 +128,21 @@ class VideoViewController : BaseViewController {
         
         headerPanGestureRecognizer.require(toFail: headerPullDownGestureRecognizer)
         
-        self.demoModeBlurView.isHidden = !self.demoMode
+        // The AR video feed & CG objects only exist in demo mode. In normal operation, they are replaced by a live feed of the Unreal Engine rendering, with the virtual camera's position determined by the device & ARKit.
+        //self.demoModeBlurView.isHidden = !self.demoMode
         showReconnecting(false, animated: false)
         
         let config = ARWorldTrackingConfiguration()
-
         config.worldAlignment = .gravity
         
+        arSession = ARSession()
+        arSession?.delegate = self
+        arSession?.run(config)
+        
+        self.streamingConnection?.renderView = self.renderView
+        
+        /*
+        
         arView.session.run(config, options: [ .resetTracking ] )
         arView.session.delegate = self
         arView.delegate = self
@@ -232,19 +182,20 @@ class VideoViewController : BaseViewController {
             arView.scene.background.contents = UIColor.black
         }
         
-        self.initMetal(arView)
+         */
         
         let coachingOverlayView = ARCoachingOverlayView()
         
         self.view.insertSubview(coachingOverlayView, belowSubview: self.reconnectingBlurView)
         coachingOverlayView.layout(.left, .right, .top, .bottom, to: self.arView)
         coachingOverlayView.goal = .tracking
-        coachingOverlayView.session = self.arView.session
+        coachingOverlayView.session = self.arSession
         coachingOverlayView.activatesAutomatically = true
         coachingOverlayView.delegate = self
 
-        reconnect()
-
+        statsTimer = Timer.scheduledTimer(withTimeInterval: 1.0, repeats: true, block: { t in
+            self.updateStreamingStats()
+        })
     }
     
     override func viewDidAppear(_ animated: Bool) {
@@ -282,7 +233,7 @@ class VideoViewController : BaseViewController {
         // if still connected, we need to send a disconnect/goodbye, then dismiss the view controller when
         // the goodbye msg is sent.
         // otherwise, we can call disconnect (which tears down some objects) and dismiss this VC immediately.
-        if self.oscConnection?.isConnected ?? false {
+        if self.streamingConnection?.isConnected ?? false {
             dismissOnDisconnect = true
             disconnect()
         } else {
@@ -292,272 +243,42 @@ class VideoViewController : BaseViewController {
     }
     
     func reconnect() {
-        
-        decoder = JPEGVideoDecoder()
-
-        pingTimer = Timer.scheduledTimer(withTimeInterval: 1, repeats: true, block: { (t) in
-            
-            if let conn = self.oscConnection {
-                if conn.isConnected {
-                    conn.send(OSCAddressPattern.ping)
-                } else {
-                    conn.reconnect()
-                }
-            }
-        })
-        
-        statsTimer = Timer.scheduledTimer(withTimeInterval: 1.0, repeats: true, block: { t in
-            self.updateStreamingStats()
-        })
-        
-        oscConnection?.reconnect()
+        self.streamingConnection?.reconnect()
     }
     
     func disconnect() {
+        self.streamingConnection?.disconnect()
         
-        pingTimer?.invalidate()
-        pingTimer = nil
-        
-        statsTimer?.invalidate()
-        statsTimer = nil
-        
-        oscConnection?.disconnect()
-        
-        self.decoder = nil
-    }
-
-    func makeMTLCommandBuffer() -> MTLCommandBuffer? {
-        return mtlCommandQueue?.makeCommandBuffer()
-    }
-    
-    func decode(_ width : Int32, _ height : Int32, _ blob : Data) {
-
-        let sz = blob.count
-        
-        self.decoder?.decode(width: width, height: height, data: blob) { (pixelBuffer) in
-            
-            if let pb = pixelBuffer {
-                
-                
-                DispatchQueue.main.async {
-
-                    self.pixelBufferSize = CGSize(width: CVPixelBufferGetWidth(pb), height: CVPixelBufferGetHeight(pb))
-                    self.stats.append((Date(), sz))
-                    
-                    self.headerViewY = self.headerView.layer.presentation()?.frame.origin.y ?? 0
-                }
-                
-                self.decodedPixelBuffer = pb
-            }
-            
-        }
     }
     
     func updateStreamingStats() {
-        
-        let now = Date()
-        var countToRemove = 0
-        var totalBytesLastSecond = 0
-        for item in stats {
-            if now.timeIntervalSince(item.0) > 1 {
-                countToRemove += 1
-            } else {
-                totalBytesLastSecond += item.1
-            }
-        }
-
-        self.stats.removeFirst(countToRemove)
 
-        var fps : Float?
-        if stats.count > 1,
-           let first = stats.first,
-           let last = stats.last {
-            
-            fps = Float(stats.count - 1) / Float(last.0.timeIntervalSince(first.0))
-        }
-        
         var str = ""
         
-        if let sz = self.pixelBufferSize {
+        if let sz = self.streamingConnection?.videoSize {
             str += "\(Int(sz.width))x\(Int(sz.height))"
         }
         
-        if fps != nil {
-            str += String(format: " %0.2f FPS", fps!)
-        }
-
-        
-        if !str.isEmpty {
-            str += " • "
-        }
-        
-        ByteCountFormatter().string(fromByteCount: Int64(totalBytesLastSecond))
-        str += ByteCountFormatter().string(fromByteCount: Int64(totalBytesLastSecond)) + "/sec"
-
-        self.headerView.stats = str
-    }
-    
-    
-    func initMetal(_ arView : ARSCNView) {
-        
-        guard let mtlDevice = arView.device else {
-            Log.error("MTLCreateSystemDefaultDevice failed!")
-            return
-        }
-
-        // Load all the shader files with a .metal file extension in the project.
-        guard let defaultLibrary = mtlDevice.makeDefaultLibrary() else {
-            Log.error("makeDefaultLibrary failed!")
-            return
-        }
-        
-        let vertexFunction = defaultLibrary.makeFunction(name: "passthrough_vertex")
-        let fragmentFunction = defaultLibrary.makeFunction(name: "basic_texture")
-
-        // Configure a pipeline descriptor that is used to create a pipeline state.
-        let pipelineDescriptor = MTLRenderPipelineDescriptor()
-        pipelineDescriptor.label = "Passthrough Pipeline";
-        pipelineDescriptor.vertexFunction = vertexFunction;
-        pipelineDescriptor.fragmentFunction = fragmentFunction;
-        pipelineDescriptor.colorAttachments[0].pixelFormat = arView.colorPixelFormat;
-        pipelineDescriptor.depthAttachmentPixelFormat = arView.depthPixelFormat
-
-        do {
-            mtlPassthroughPipelineState = try mtlDevice.makeRenderPipelineState(descriptor: pipelineDescriptor)
-        } catch {
-            Log.error("makeRenderPipelineState failed : \(error.localizedDescription)")
-
-        }
-        
-        mtlCommandQueue = mtlDevice.makeCommandQueue()
-        
-        let vertices : [Float] = [ -1.0,  1.0, 0.0, 1.0,
-                                   -1.0, -1.0, 0.0, 0.0,
-                                    1.0,  1.0, 1.0, 1.0,
-                                    1.0, -1.0, 1.0, 0.0 ]
-
-        mtlVertexBuffer = mtlDevice.makeBuffer(bytes: vertices, length: MemoryLayout<Float>.size * vertices.count, options: .storageModeShared)
-        
-
-        let textureLoader = MTKTextureLoader(device: mtlDevice)
-        
-        do {
-            mtlNoSignalTexture = try textureLoader.newTexture(name: "no_signal", scaleFactor: 1.0, bundle: nil, options: nil)
-        } catch {
-            Log.error("couldn't load the no_signal texture: \(error.localizedDescription)")
-        }
-        
-        let status = CVMetalTextureCacheCreate(kCFAllocatorDefault, nil, mtlDevice, nil, &mtlTextureCache)
-        Log.info("CVMetalTextureCacheCreate : \(status)")
-        
-    }
-    
-    func sendOSCAddressPattern(_ pattern : OSCAddressPattern, point : CGPoint, finger : Int, force : CGFloat) {
-
-        Log.debug("\(pattern.rawValue) finger \(finger) @ \(point.x),\(point.y)")
-        
-        guard let textureAspectRatio = self.pixelBufferAspectRatio else { return }
-
-        // we need to adjust for the size of the top toolbar in the view : this toolbar is *over* the
-        // rendering/arView : if we resize the arView (which is a SCNView) with animated constraints,
-        // there are all kinds of visual side effects.
-        let toolbarHeightInView = (self.headerViewY + self.headerViewHeight)
-        
-        let viewportSize = CGSize(width: self.arView.frame.width, height: self.arView.frame.height - toolbarHeightInView)
-        let viewportAspectRatio = Float(viewportSize.width) / Float(viewportSize.height)
-
-        // offset the touch point with respect to the top toolbar
-        let pointInViewport = CGPoint(x: point.x, y: point.y - toolbarHeightInView)
-        var normalizedPoint = CGPoint(x: pointInViewport.x / viewportSize.width, y: pointInViewport.y / viewportSize.height)
-
-        
-        if textureAspectRatio > viewportAspectRatio {
-            
-            normalizedPoint.y = (((normalizedPoint.y * 2.0 - 1.0) * CGFloat(textureAspectRatio / viewportAspectRatio)) + 1.0) / 2.0;
-
-        } else if textureAspectRatio < viewportAspectRatio {
-
-            normalizedPoint.x = (((normalizedPoint.x * 2.0 - 1.0) * CGFloat(viewportAspectRatio / textureAspectRatio)) + 1.0) / 2.0;
-        }
-        
-        // reject any touches that aren't in the 0.0 -> 1.0 of the remote image
-        if (normalizedPoint.x < 0.0) ||
-            (normalizedPoint.x > 1.0) ||
-            (normalizedPoint.y < 0.0) ||
-            (normalizedPoint.y > 1.0) {
-            return
-        }
-
-        //Log.info("\(pattern.rawValue) : \(normalizedPoint.x), \(normalizedPoint.y)")
-        let data = OSCUtility.ueTouchData(point: normalizedPoint, finger: finger, force: force)
-        oscConnection?.send(pattern, arguments: [ OSCArgument.blob(data) ])
-    }
-    
-    override func touchesBegan(_ touches: Set<UITouch>, with event: UIEvent?) {
-
-        guard relayTouchEvents else { return }
-        
-        // for each touch, we need to maintain an index to send via OSC, so that UE
-        // can correctly map between began/moved. The currentTouches array will maintain
-        // the indices for each touch (keeping some nil if needed) and clean up / compress
-        // the array when a touch ends.
-        
-        for touch in touches {
-
-            // insert this new touch into the array, using the first available slot (a nil entry)
-            // or append to the end.
-            var fingerIndex = 0
-            
-            for i in 0...self.currentTouches.count {
-                fingerIndex = i
-                if i == self.currentTouches.count {
-                    self.currentTouches.append(touch)
-                    break
-                } else if self.currentTouches[i] == nil {
-                    self.currentTouches[i] = touch
-                    break
+        if let stats = self.streamingConnection?.stats {
+                
+            if let bps = stats.bytesPerSecond {
+                if !str.isEmpty {
+                    str += " • "
                 }
+             
+                str += ByteCountFormatter().string(fromByteCount: Int64(bps)) + "/sec"
             }
             
-            self.sendOSCAddressPattern(OSCAddressPattern.touchStarted, point: touch.location(in: self.arView), finger: fingerIndex, force: touch.force)
-        }
-    }
-    
-    override func touchesMoved(_ touches: Set<UITouch>, with event: UIEvent?) {
-
-        guard relayTouchEvents else { return }
-
-        for touch in touches {
-            if let fingerIndex = self.currentTouches.firstIndex(of: touch) {
-                self.sendOSCAddressPattern(OSCAddressPattern.touchMoved, point: touch.location(in: self.arView), finger: fingerIndex, force: touch.force)
+            if let fps = stats.framesPerSecond {
+                if !str.isEmpty {
+                    str += " • "
+                }
+             
+                str += "\(fps) fps"
             }
         }
 
-    }
-    
-    override func touchesEnded(_ touches: Set<UITouch>, with event: UIEvent?) {
-
-        guard relayTouchEvents else { return }
-
-        for touch in touches {
-            
-            
-            if let fingerIndex = self.currentTouches.firstIndex(of: touch) {
-                self.sendOSCAddressPattern(OSCAddressPattern.touchEnded, point: touch.location(in: self.arView), finger: fingerIndex, force: touch.force)
-                
-                // this touch is ended, set to nil
-                self.currentTouches[fingerIndex] = nil
-            }
-        }
-        
-        // compress the array : remove all nils from the end
-        if let lastNonNilIndex = self.currentTouches.lastIndex(where: { $0 != nil }) {
-            if lastNonNilIndex < (self.currentTouches.count - 1)  {
-                self.currentTouches.removeSubrange((lastNonNilIndex + 1)..<self.currentTouches.count)
-            }
-        } else {
-            self.currentTouches.removeAll()
-        }
+        self.headerView.stats = str
     }
 
     @objc func handleHeaderPanGesture(_ gesture : UIGestureRecognizer) {
@@ -629,18 +350,24 @@ class VideoViewController : BaseViewController {
     
     func sendThumbstickUpdate(key : String, controller : Int, oldValue : Float?, newValue : Float) {
 
+        assertionFailure()
+        /*
         guard let osc = oscConnection else { return }
 
         if (newValue != 0.0) || ((oldValue ?? 0.0) != newValue) {
             osc.send(.controllerAnalog, arguments: [ OSCArgument.blob(OSCUtility.ueControllerAnalogData(key: key, controller: controller, value: newValue)) ])
         }
+         */
     }
     
     func sendControllerUpdate() {
         
-        guard let _ = oscConnection else { return }
         guard let gc = gameController else { return }
+        assertionFailure()
         
+        /*
+        guard let _ = oscConnection else { return }
+
         let snapshot = gc.capture()
         guard let gp = snapshot.extendedGamepad else { return }
 
@@ -652,9 +379,12 @@ class VideoViewController : BaseViewController {
         sendThumbstickUpdate(key: "Gamepad_RightY", controller: controllerIndex, oldValue : self.gameControllerSnapshot?.extendedGamepad?.rightThumbstick.yAxis.value, newValue :gp.rightThumbstick.yAxis.value)
         
         self.gameControllerSnapshot = snapshot
+         */
     }
 
     func directionalPadValueChanged(directionalPad : GCControllerDirectionPad, xValue : Float, yValue : Float) {
+        assertionFailure()
+        /*
         guard let osc = oscConnection else { return }
         guard let gc = gameController else { return }
         guard let gp = gc.extendedGamepad else { return }
@@ -672,10 +402,13 @@ class VideoViewController : BaseViewController {
             Log.warning("directionalPadValueChanged() encounted unsupported '\(directionalPad.localizedName ?? "UNK")'")
             break
         }
+         */
     }
     
     func buttonValueChanged(button : GCControllerButtonInput, value : Float, pressed : Bool) {
 
+        assertionFailure()
+        /*
         guard let osc = oscConnection else { return }
         guard let gc = gameController else { return }
         guard let gp = gc.extendedGamepad else { return }
@@ -726,7 +459,7 @@ class VideoViewController : BaseViewController {
         default:
             Log.warning("buttonValueChanged() encounted unsupported '\(button.localizedName ?? "UNK")'")
         }
-
+         */
     }
 }
 
diff --git a/Engine/Extras/VirtualProduction/LiveLinkVCAM/vcam/WebRTC/Extensions/IceCandidate.swift b/Engine/Extras/VirtualProduction/LiveLinkVCAM/vcam/WebRTC/Extensions/IceCandidate.swift
new file mode 100644
--- /dev/null
+++ b/Engine/Extras/VirtualProduction/LiveLinkVCAM/vcam/WebRTC/Extensions/IceCandidate.swift
@@ -0,0 +1,27 @@
+//
+//  IceCandidate.swift
+//  WebRTC-Demo
+//
+//  Created by Stasel on 20/02/2019.
+//  Copyright © 2019 Stasel. All rights reserved.
+//
+
+import Foundation
+import WebRTC
+
+/// This struct is a swift wrapper over `RTCIceCandidate` for easy encode and decode
+struct IceCandidate: Codable {
+    let candidate: String
+    let sdpMLineIndex: Int32
+    let sdpMid: String?
+    
+    init(from iceCandidate: RTCIceCandidate) {
+        self.sdpMLineIndex = iceCandidate.sdpMLineIndex
+        self.sdpMid = iceCandidate.sdpMid
+        self.candidate = iceCandidate.sdp
+    }
+    
+    var rtcIceCandidate: RTCIceCandidate {
+        return RTCIceCandidate(sdp: self.candidate, sdpMLineIndex: self.sdpMLineIndex, sdpMid: self.sdpMid)
+    }
+}
diff --git a/Engine/Extras/VirtualProduction/LiveLinkVCAM/vcam/WebRTC/Extensions/IntToBytes.swift b/Engine/Extras/VirtualProduction/LiveLinkVCAM/vcam/WebRTC/Extensions/IntToBytes.swift
new file mode 100644
--- /dev/null
+++ b/Engine/Extras/VirtualProduction/LiveLinkVCAM/vcam/WebRTC/Extensions/IntToBytes.swift
@@ -0,0 +1,44 @@
+//
+//  IntToBytes.swift
+//  WebRTC-Demo
+//
+//  Created by Luke Bermingham on 20/6/2022.
+//  Very similar to this answer on StackOverflow: https://stackoverflow.com/a/41368211
+//
+
+import Foundation
+
+// Add the `toBytes` method to UInt16, UInt32, and UInt64
+protocol UIntToBytesConvertable {
+    func toBytes() -> [UInt8]
+}
+
+extension UIntToBytesConvertable {
+    func toByteArr<T: UnsignedInteger>(endian: T, count: Int) -> [UInt8] {
+        var _endian = endian
+        let bytePtr = withUnsafePointer(to: &_endian) {
+            $0.withMemoryRebound(to: UInt8.self, capacity: count) {
+                UnsafeBufferPointer(start: $0, count: count)
+            }
+        }
+        return [UInt8](bytePtr)
+    }
+}
+
+extension UInt16: UIntToBytesConvertable {
+    func toBytes() -> [UInt8] {
+        return toByteArr(endian: self.littleEndian, count: MemoryLayout<UInt16>.size)
+    }
+}
+
+extension UInt32: UIntToBytesConvertable {
+    func toBytes() -> [UInt8] {
+        return toByteArr(endian: self.littleEndian, count: MemoryLayout<UInt32>.size)
+    }
+}
+
+extension UInt64: UIntToBytesConvertable {
+    func toBytes() -> [UInt8] {
+        return toByteArr(endian: self.littleEndian, count: MemoryLayout<UInt64>.size)
+    }
+}
diff --git a/Engine/Extras/VirtualProduction/LiveLinkVCAM/vcam/WebRTC/Extensions/PeerConnectionConfig.swift b/Engine/Extras/VirtualProduction/LiveLinkVCAM/vcam/WebRTC/Extensions/PeerConnectionConfig.swift
new file mode 100644
--- /dev/null
+++ b/Engine/Extras/VirtualProduction/LiveLinkVCAM/vcam/WebRTC/Extensions/PeerConnectionConfig.swift
@@ -0,0 +1,47 @@
+//
+//  PeerConnectionConfig.swift
+//  WebRTC-Demo
+//
+//  Created by Luke Bermingham on 10/6/2022.
+//  Copyright © 2022 Stas Seldin. All rights reserved.
+//
+
+import Foundation
+import WebRTC
+
+// Represents peer connection configuration options such as ICE servers being used.
+struct PeerConnectionConfig: Codable {
+    let username: String?
+    let credential: String?
+    let urls: [String]?
+    
+    init(from rtcConfiguration: RTCConfiguration) {
+        
+        if !rtcConfiguration.iceServers.isEmpty {
+            self.urls = rtcConfiguration.iceServers[0].urlStrings
+            self.username = rtcConfiguration.iceServers[0].username
+            self.credential = rtcConfiguration.iceServers[0].credential
+        }
+        else {
+            self.urls = []
+            self.username = ""
+            self.credential = ""
+        }
+    }
+    
+    var rtcConfiguration: RTCConfiguration {
+               
+        let rtcConfig : RTCConfiguration = RTCConfiguration()
+        
+        if let urls = urls, let username = username, let credential = credential {
+            rtcConfig.iceServers = [RTCIceServer(urlStrings: urls, username: username, credential: credential)]
+        }
+        else if let urls = urls {
+            rtcConfig.iceServers = [RTCIceServer(urlStrings: urls)]
+        }
+        
+        rtcConfig.sdpSemantics = .unifiedPlan;
+        rtcConfig.continualGatheringPolicy = .gatherContinually
+        return rtcConfig
+    }
+}
diff --git a/Engine/Extras/VirtualProduction/LiveLinkVCAM/vcam/WebRTC/Extensions/PixelStreamingToClientMessage.swift b/Engine/Extras/VirtualProduction/LiveLinkVCAM/vcam/WebRTC/Extensions/PixelStreamingToClientMessage.swift
new file mode 100644
--- /dev/null
+++ b/Engine/Extras/VirtualProduction/LiveLinkVCAM/vcam/WebRTC/Extensions/PixelStreamingToClientMessage.swift
@@ -0,0 +1,26 @@
+//
+//  PixelStreamingToClientMessage.swift
+//  WebRTC-Demo
+//
+//  Created by Luke Bermingham on 15/6/2022.
+//  Copyright © 2022 Stas Seldin. All rights reserved.
+//
+
+import Foundation
+
+// Must stay synced with Pixel Streaming plugin until a better solution is devised.
+enum PixelStreamingToClientMessage: UInt8, Codable {
+    case QualityControlOwnership = 0
+    case Response = 1
+    case Command = 2
+    case FreezeFrame = 3
+    case UnfreezeFrame = 4
+    case VideoEncoderAvgQP = 5
+    case LatencyTest = 6
+    case InitialSettings = 7
+    case FileExtension = 8
+    case FileMimeType = 9
+    case FileContents = 10
+}
+
+
diff --git a/Engine/Extras/VirtualProduction/LiveLinkVCAM/vcam/WebRTC/Extensions/PixelStreamingToStreamerMessage.swift b/Engine/Extras/VirtualProduction/LiveLinkVCAM/vcam/WebRTC/Extensions/PixelStreamingToStreamerMessage.swift
new file mode 100644
--- /dev/null
+++ b/Engine/Extras/VirtualProduction/LiveLinkVCAM/vcam/WebRTC/Extensions/PixelStreamingToStreamerMessage.swift
@@ -0,0 +1,44 @@
+//
+//  PixelStreamingToStreamerMessage.swift
+//  WebRTC-Demo
+//
+//  Created by Luke Bermingham on 20/6/2022.
+//
+
+import Foundation
+
+enum PixelStreamingToStreamerMessage : UInt8, Codable {
+    case IFrameRequest = 0
+    case RequestQualityControl = 1
+    case FpsRequest = 2
+    case AverageBitrateRequest = 3
+    case StartStreaming = 4
+    case StopStreaming = 5
+    case LatencyTest = 6
+    case RequestInitialSettings = 7
+    case TestEcho = 8
+    
+    case UIInteraction = 50
+    case Command = 51
+
+    case KeyDown = 60
+    case KeyUp = 61
+    case KeyPress = 62
+
+    case MouseEnter = 70
+    case MouseLeave = 71
+    case MouseDown = 72
+    case MouseUp = 73
+    case MouseMove = 74
+    case MouseWheel = 75
+
+    case TouchStart = 80
+    case TouchEnd = 81
+    case TouchMove = 82
+
+    case GamepadButtonPressed = 90
+    case GamepadButtonReleased = 91
+    case GamepadAnalog = 92
+    
+    case Transform = 100
+}
diff --git a/Engine/Extras/VirtualProduction/LiveLinkVCAM/vcam/WebRTC/Extensions/RTCStates.swift b/Engine/Extras/VirtualProduction/LiveLinkVCAM/vcam/WebRTC/Extensions/RTCStates.swift
new file mode 100644
--- /dev/null
+++ b/Engine/Extras/VirtualProduction/LiveLinkVCAM/vcam/WebRTC/Extensions/RTCStates.swift
@@ -0,0 +1,64 @@
+//
+//  RTCConnectionState.swift
+//  WebRTC
+//
+//  Created by Stasel on 20/05/2018.
+//  Copyright © 2018 Stasel. All rights reserved.
+//
+
+import Foundation
+import WebRTC
+
+extension RTCIceConnectionState: CustomStringConvertible {
+    public var description: String {
+        switch self {
+        case .new:          return "new"
+        case .checking:     return "checking"
+        case .connected:    return "connected"
+        case .completed:    return "completed"
+        case .failed:       return "failed"
+        case .disconnected: return "disconnected"
+        case .closed:       return "closed"
+        case .count:        return "count"
+        @unknown default:   return "Unknown \(self.rawValue)"
+        }
+    }
+}
+
+extension RTCSignalingState: CustomStringConvertible {
+    public var description: String {
+        switch self {
+        case .stable:               return "stable"
+        case .haveLocalOffer:       return "haveLocalOffer"
+        case .haveLocalPrAnswer:    return "haveLocalPrAnswer"
+        case .haveRemoteOffer:      return "haveRemoteOffer"
+        case .haveRemotePrAnswer:   return "haveRemotePrAnswer"
+        case .closed:               return "closed"
+        @unknown default:   return "Unknown \(self.rawValue)"
+        }
+    }
+}
+
+extension RTCIceGatheringState: CustomStringConvertible {
+    public var description: String {
+        switch self {
+        case .new:          return "new"
+        case .gathering:    return "gathering"
+        case .complete:     return "complete"
+        @unknown default:   return "Unknown \(self.rawValue)"
+        }
+    }
+}
+
+extension RTCDataChannelState: CustomStringConvertible {
+    public var description: String {
+        switch self {
+        case .connecting:   return "connecting"
+        case .open:         return "open"
+        case .closing:      return "closing"
+        case .closed:       return "closed"
+        @unknown default:   return "Unknown \(self.rawValue)"
+        }
+    }
+}
+
diff --git a/Engine/Extras/VirtualProduction/LiveLinkVCAM/vcam/WebRTC/Extensions/SessionDescription.swift b/Engine/Extras/VirtualProduction/LiveLinkVCAM/vcam/WebRTC/Extensions/SessionDescription.swift
new file mode 100644
--- /dev/null
+++ b/Engine/Extras/VirtualProduction/LiveLinkVCAM/vcam/WebRTC/Extensions/SessionDescription.swift
@@ -0,0 +1,52 @@
+//
+//  SessionDescription.swift
+//  WebRTC-Demo
+//
+//  Created by Stasel on 20/02/2019.
+//  Copyright © 2019 Stasel. All rights reserved.
+//
+
+import Foundation
+import WebRTC
+
+/// This enum is a swift wrapper over `RTCSdpType` for easy encode and decode
+enum SdpType: String, Codable {
+    case offer, prAnswer, answer, rollback
+    
+    var rtcSdpType: RTCSdpType {
+        switch self {
+        case .offer:    return .offer
+        case .answer:   return .answer
+        case .prAnswer: return .prAnswer
+        case .rollback: return .rollback
+        }
+    }
+}
+
+/// This struct is a swift wrapper over `RTCSessionDescription` for easy encode and decode
+struct SessionDescription: Codable {
+    let sdp: String
+    let type: SdpType
+    
+    init(from sdpType: SdpType, and sdp: String) {
+        self.sdp = sdp
+        self.type = sdpType
+    }
+    
+    init(from rtcSessionDescription: RTCSessionDescription) {
+        self.sdp = rtcSessionDescription.sdp
+        
+        switch rtcSessionDescription.type {
+        case .offer:    self.type = .offer
+        case .prAnswer: self.type = .prAnswer
+        case .answer:   self.type = .answer
+        case .rollback: self.type = .rollback
+        @unknown default:
+            fatalError("Unknown RTCSessionDescription type: \(rtcSessionDescription.type.rawValue)")
+        }
+    }
+    
+    var rtcSessionDescription: RTCSessionDescription {
+        return RTCSessionDescription(type: self.type.rtcSdpType, sdp: self.sdp)
+    }
+}
diff --git a/Engine/Extras/VirtualProduction/LiveLinkVCAM/vcam/WebRTC/Extensions/TouchControls.swift b/Engine/Extras/VirtualProduction/LiveLinkVCAM/vcam/WebRTC/Extensions/TouchControls.swift
new file mode 100644
--- /dev/null
+++ b/Engine/Extras/VirtualProduction/LiveLinkVCAM/vcam/WebRTC/Extensions/TouchControls.swift
@@ -0,0 +1,171 @@
+//
+//  TouchControls.swift
+//  WebRTC-Demo
+//
+//  Created by Luke Bermingham on 20/6/2022.
+//
+
+import Foundation
+import UIKit
+
+// Protocol for any touch consumer to fulfil.
+protocol TouchDelegate: AnyObject {
+    func touchesBegan(_ touches: Set<UITouch>);
+    func touchesMoved(_ touches: Set<UITouch>);
+    func touchesEnded(_ touches: Set<UITouch>);
+    func touchesCancelled(_ touches: Set<UITouch>);
+    func onVideoChangedSize(size: CGSize)
+}
+
+// Represents a coordinate along a viewport axis (0...1) that has been quantized into an uint16 (0..65536).
+struct TouchPoint {
+    let x : UInt16
+    let y : UInt16
+    // Value is not in range if the touch was outside the view bounds
+    let inRange : Bool
+    let force : UInt8
+    
+    init(x : UInt16, y : UInt16, inRange : Bool, force : UInt8) {
+        self.x = x
+        self.y = y
+        self.inRange = inRange
+        self.force = force
+    }
+}
+
+final class TouchControls : TouchDelegate {
+    
+    let webRTCClient : WebRTCClient
+    let touchView : UIView
+    var videoAspectRatio : CGFloat
+    
+    // We need a way to give each touch a unique finger id that is persistent throughout
+    // the life of that touch. So we map each touch to an id [0...10] - iPads can do 11 simulataneous touches!
+    var fingers : [Int] = [10, 9, 8, 7, 6, 5, 4, 3, 2, 1, 0]
+    var fingerIds = [UITouch : Int]()
+    
+    init(_ webRTCClient : WebRTCClient, touchView: UIView) {
+        self.webRTCClient = webRTCClient
+        self.touchView = touchView
+        self.videoAspectRatio = CGFloat.zero
+    }
+    
+    func rememberTouch(_ touch : UITouch) {
+        let fingerId : Int? = self.fingers.popLast()
+        
+        if let fingerId = fingerId {
+            self.fingerIds[touch] = fingerId
+        } else {
+            debugPrint("Exhausted all touch identifiers - this shouldn't happen, it indicates a leak in tracking touch events.")
+        }
+        
+    }
+
+    func forgetTouch(_ touch : UITouch) {
+        let touchId : Int? = self.fingerIds[touch]
+        if let touchId = touchId {
+            self.fingers.append(touchId)
+            self.fingerIds.removeValue(forKey: touch)
+        } else {
+            debugPrint("Could not forget this touch because we don't have it stored.")
+        }
+    }
+    
+    func touchesBegan(_ touches: Set<UITouch>) {
+        // Assign a unique identifier to each touch.
+        for touch in touches {
+            self.rememberTouch(touch)
+            //let coord : CGPoint = touch.location(in: self.touchView)
+            //debugPrint("Touch started x=\(coord.x) y=\(coord.y)")
+        }
+        
+        self.sendTouchData(messageType: PixelStreamingToStreamerMessage.TouchStart, touches: touches)
+    }
+    
+    func touchesMoved(_ touches: Set<UITouch>) {
+        //debugPrint("Touch moved - count=\(touches.count)")
+        self.sendTouchData(messageType: PixelStreamingToStreamerMessage.TouchMove, touches: touches)
+    }
+    
+    func touchesEnded(_ touches: Set<UITouch>) {
+        //debugPrint("Touch ended - count=\(touches.count)")
+        self.sendTouchData(messageType: PixelStreamingToStreamerMessage.TouchEnd, touches: touches)
+
+        // Re-cycle unique identifiers previously assigned to each touch.
+        for touch in touches {
+            self.forgetTouch(touch)
+        }
+    }
+    
+    func touchesCancelled(_ touches: Set<UITouch>) {
+        // same as touches ended unless we are hooking up gestures
+        touchesEnded(touches)
+    }
+    
+    func onVideoChangedSize(size: CGSize) {
+        self.videoAspectRatio = size.height / size.width
+    }
+    
+    // Unsigned XY positions are normalized to the ratio (0.0..1.0) along a viewport axis and then quantized into an uint16 (0..65536).
+    func normalizeAndQuantize(_ touch: UITouch) -> TouchPoint {
+        let uiAspectRatio : CGFloat = touchView.bounds.height / touchView.bounds.width
+        let uiBigger : Bool = uiAspectRatio > videoAspectRatio
+        let ratio : CGFloat = uiAspectRatio / videoAspectRatio;
+        let touchLocation : CGPoint = touch.location(in: touchView)
+        let viewBounds : CGRect = touchView.bounds
+        
+        // normalise x,y to the UI element bounds
+        let normalizedX : CGFloat = uiBigger ? (touchLocation.x / viewBounds.width) : (ratio * (touchLocation.x / viewBounds.width - 0.5) + 0.5)
+        let normalizedY : CGFloat = uiBigger ? (ratio * (touchLocation.y / viewBounds.height - 0.5) + 0.5) : (touchLocation.y / viewBounds.height)
+        
+        // normalize force value of touch
+        let normalizedForce : UInt8 = touch.maximumPossibleForce > 0 ? UInt8(touch.force / touch.maximumPossibleForce * CGFloat(UInt8.max)) : 1
+        
+        // Detect if touch is out of bounds
+        if normalizedX < 0.0 || normalizedX > 1.0 || normalizedY < 0.0 || normalizedY > 1.0 {
+            return TouchPoint(x: UInt16.max, y: UInt16.max, inRange: false, force: 0)
+        } else {
+            return TouchPoint(x: UInt16(normalizedX * CGFloat(UInt16.max)), y: UInt16(normalizedY * CGFloat(UInt16.max)), inRange: true, force: normalizedForce)
+        }
+    }
+    
+    func sendTouchData(messageType: PixelStreamingToStreamerMessage, touches: Set<UITouch>) {
+        
+        var bytes: [UInt8] = []
+        
+        // Add type of touch using 1 byte
+        bytes.append(messageType.rawValue)
+        
+        // Add number of touches
+        bytes.append(UInt8(touches.count))
+        
+        for touch in touches {
+            let touchId : Int? = fingerIds[touch]
+            
+            if let touchId = touchId {
+                let convertedTouch : TouchPoint = normalizeAndQuantize(touch)
+                
+                // Write X coordinate as two bytes e.g. UInt16 -> [UInt8]
+                bytes.append(contentsOf: convertedTouch.x.toBytes())
+                
+                // Write Y coorindate as two bytes e.g. UInt16 -> [UInt8]
+                bytes.append(contentsOf: convertedTouch.y.toBytes())
+                
+                // Write finger id as one byte
+                bytes.append(UInt8(touchId))
+                
+                // Write normalized force as one byte
+                bytes.append(convertedTouch.force)
+                
+                // Write in range as one byte
+                bytes.append(convertedTouch.inRange ? 1 : 0)
+            } else {
+                debugPrint("Could not send touch because did not have any matching touch id for it.")
+            }
+        }
+        
+        let data = Data(bytes)
+        self.webRTCClient.sendData(data)
+    }
+    
+}
diff --git a/Engine/Extras/VirtualProduction/LiveLinkVCAM/vcam/WebRTC/Services/Message.swift b/Engine/Extras/VirtualProduction/LiveLinkVCAM/vcam/WebRTC/Services/Message.swift
new file mode 100644
--- /dev/null
+++ b/Engine/Extras/VirtualProduction/LiveLinkVCAM/vcam/WebRTC/Services/Message.swift
@@ -0,0 +1,63 @@
+//
+//  Message.swift
+//  WebRTC-Demo
+//
+//  Created by Stasel on 20/02/2019.
+//  Copyright © 2019 Stasel. All rights reserved.
+//
+
+import Foundation
+
+enum Message {
+    case sdp(SessionDescription)
+    case candidate(IceCandidate)
+    case config(PeerConnectionConfig)
+    case playerCount(Int)
+}
+
+extension Message: Codable {
+    init(from decoder: Decoder) throws {
+        let container = try decoder.container(keyedBy: CodingKeys.self)
+        let type = try container.decode(String.self, forKey: .type)
+        switch type {
+        case "config":
+            self = .config(try container.decode(PeerConnectionConfig.self, forKey: .peerConnectionOptions))
+        case "offer":
+            let sdpStr = try container.decode(String.self, forKey: .sdp)
+            let sdpType = try container.decode(SdpType.self, forKey: .type)
+            self = .sdp(SessionDescription(from: sdpType, and: sdpStr))
+        case "iceCandidate":
+            self = .candidate(try container.decode(IceCandidate.self, forKey: .candidate))
+        case "playerCount":
+            self = .playerCount(try container.decode(Int.self, forKey: .count))
+        default:
+            throw DecodeError.unknownType
+        }
+    }
+    
+    func encode(to encoder: Encoder) throws {
+        var container = encoder.container(keyedBy: CodingKeys.self)
+        switch self {
+        case .config(let config):
+            try container.encode("config", forKey: .type)
+            try container.encode(config, forKey: .peerConnectionOptions)
+        case .sdp(let answerSdp):
+            try container.encode("answer", forKey: .type)
+            try container.encode(answerSdp.sdp, forKey: .sdp)
+        case .candidate(let candidate):
+            try container.encode("iceCandidate", forKey: .type)
+            try container.encode(candidate, forKey: .candidate)
+        case .playerCount(let count):
+            try container.encode("playerCount", forKey: .type)
+            try container.encode(count, forKey: .count)
+        }
+    }
+    
+    enum DecodeError: Error {
+        case unknownType
+    }
+    
+    enum CodingKeys: String, CodingKey {
+        case type, sdp, candidate, peerConnectionOptions, count
+    }
+}
diff --git a/Engine/Extras/VirtualProduction/LiveLinkVCAM/vcam/WebRTC/Services/SignalingClient.swift b/Engine/Extras/VirtualProduction/LiveLinkVCAM/vcam/WebRTC/Services/SignalingClient.swift
new file mode 100644
--- /dev/null
+++ b/Engine/Extras/VirtualProduction/LiveLinkVCAM/vcam/WebRTC/Services/SignalingClient.swift
@@ -0,0 +1,120 @@
+//
+//  SignalClient.swift
+//  WebRTC
+//
+//  Created by Stasel on 20/05/2018.
+//  Copyright © 2018 Stasel. All rights reserved.
+//
+
+import Foundation
+import WebRTC
+
+protocol SignalClientDelegate: AnyObject {
+    func signalClientDidConnect(_ signalClient: SignalingClient)
+    func signalClientDidDisconnect(_ signalClient: SignalingClient, error: Error?)
+    func signalClientDidReceiveError(_ signalClient: SignalingClient, error: Error?)
+    func signalClient(_ signalClient: SignalingClient, didReceiveRemoteSdp sdp: RTCSessionDescription)
+    func signalClient(_ signalClient: SignalingClient, didReceiveCandidate candidate: RTCIceCandidate)
+    func signalClient(_ signalClient: SignalingClient, didReceiveConfig config: RTCConfiguration)
+}
+
+final class SignalingClient {
+    
+    private let decoder = JSONDecoder()
+    private let encoder = JSONEncoder()
+    private let webSocket: WebSocketProvider
+    private var shouldReconnect : Bool = true
+    weak var delegate: SignalClientDelegate?
+    
+    init(webSocket: WebSocketProvider) {
+        self.webSocket = webSocket
+    }
+    
+    func connect() {
+        self.webSocket.delegate = self
+        self.webSocket.connect()
+    }
+    
+    func close() {
+        self.shouldReconnect = false
+        self.webSocket.close()
+        self.webSocket.delegate = nil
+    }
+    
+    func send(sdp rtcSdp: RTCSessionDescription) {
+        let message = Message.sdp(SessionDescription(from: rtcSdp))
+        do {
+            let dataMessage = try self.encoder.encode(message)
+            
+            self.webSocket.send(data: dataMessage)
+        }
+        catch {
+            debugPrint("Warning: Could not encode sdp: \(error)")
+        }
+    }
+    
+    func send(candidate rtcIceCandidate: RTCIceCandidate) {
+        let message = Message.candidate(IceCandidate(from: rtcIceCandidate))
+        do {
+            let dataMessage = try self.encoder.encode(message)
+            self.webSocket.send(data: dataMessage)
+        }
+        catch {
+            debugPrint("Warning: Could not encode candidate: \(error)")
+        }
+    }
+}
+
+
+extension SignalingClient: WebSocketProviderDelegate {
+    
+    func webSocketDidConnect(_ webSocket: WebSocketProvider) {
+        self.delegate?.signalClientDidConnect(self)
+    }
+    
+    func webSocketDidDisconnect(_ webSocket: WebSocketProvider, error: Error?) {
+        self.delegate?.signalClientDidDisconnect(self, error: error)
+        
+        // try to reconnect every two seconds
+        DispatchQueue.global().asyncAfter(deadline: .now() + 2) {
+            if self.shouldReconnect {
+                debugPrint("Trying to reconnect to signaling server...")
+                self.webSocket.connect()
+            }
+        }
+    }
+    
+    func webSocketReceiveError(_ webSocket: WebSocketProvider, error: Error?) {
+        self.delegate?.signalClientDidReceiveError(self, error: error)
+    }
+    
+    // Handle a message sent over websocket
+    func webSocketDidReceiveData(_ webSocket: WebSocketProvider, data: Data) {
+        
+        // Unescape string
+        var dataStr = String(decoding: data, as: UTF8.self)
+        dataStr = dataStr.replacingOccurrences(of: "\\\"", with: "\"")
+        
+        let message: Message
+        do {
+            message = try self.decoder.decode(Message.self, from: dataStr.data(using: .utf8)!)
+        }
+        catch {
+            let dataStr = String(decoding: data, as: UTF8.self)
+            debugPrint("Warning: Could not decode incoming message: \(error) - Message= \(dataStr)")
+            return
+        }
+        
+        switch message {
+        case .config(let config):
+            self.delegate?.signalClient(self, didReceiveConfig: config.rtcConfiguration)
+        case .candidate(let iceCandidate):
+            self.delegate?.signalClient(self, didReceiveCandidate: iceCandidate.rtcIceCandidate)
+        case .sdp(let sessionDescription):
+            self.delegate?.signalClient(self, didReceiveRemoteSdp: sessionDescription.rtcSessionDescription)
+        case .playerCount(let playerCount):
+            print("Got player count=\(playerCount)")
+        }
+    }
+    
+}
diff --git a/Engine/Extras/VirtualProduction/LiveLinkVCAM/vcam/WebRTC/Services/WebRTCClient.swift b/Engine/Extras/VirtualProduction/LiveLinkVCAM/vcam/WebRTC/Services/WebRTCClient.swift
new file mode 100644
--- /dev/null
+++ b/Engine/Extras/VirtualProduction/LiveLinkVCAM/vcam/WebRTC/Services/WebRTCClient.swift
@@ -0,0 +1,306 @@
+//
+//  WebRTCClient.swift
+//  WebRTC
+//
+//  Created by Stasel on 20/05/2018.
+//  Heavily modified by Luke Bermingham 2022
+//  Copyright © 2018 Stasel. All rights reserved.
+//
+
+import Foundation
+import WebRTC
+
+protocol WebRTCClientDelegate: AnyObject {
+    func webRTCClient(_ client: WebRTCClient, didDiscoverLocalCandidate candidate: RTCIceCandidate)
+    func webRTCClient(_ client: WebRTCClient, didChangeConnectionState state: RTCIceConnectionState)
+    func webRTCClient(_ client: WebRTCClient, didReceiveData data: Data)
+    func webRTCClient(_ client: WebRTCClient, onStartReceiveVideo video: RTCVideoTrack)
+}
+
+final class WebRTCClient: NSObject {
+    
+    // The `RTCPeerConnectionFactory` is in charge of creating new RTCPeerConnection instances.
+    // A new RTCPeerConnection should be created every new call, but the factory is shared.
+    private static let factory: RTCPeerConnectionFactory = {
+        RTCInitializeSSL()
+        let videoEncoderFactory = RTCDefaultVideoEncoderFactory()
+        let videoDecoderFactory = RTCDefaultVideoDecoderFactory()
+        return RTCPeerConnectionFactory(encoderFactory: videoEncoderFactory, decoderFactory: videoDecoderFactory)
+    }()
+    
+    weak var delegate: WebRTCClientDelegate?
+
+    private let rtcAudioSession =  RTCAudioSession.sharedInstance()
+    private let audioQueue = DispatchQueue(label: "audio")
+    private let mediaConstrains = [kRTCMediaConstraintsOfferToReceiveAudio: kRTCMediaConstraintsValueTrue,
+                                   kRTCMediaConstraintsOfferToReceiveVideo: kRTCMediaConstraintsValueTrue]    
+    
+    private var peerConnection: RTCPeerConnection?
+    private var remoteVideoTrack: RTCVideoTrack?
+    private var dataChannel: RTCDataChannel?
+    
+    required override init() {
+        super.init()
+    }
+    
+    // MARK: PeerConnection Config
+    func setupPeerConnection(rtcConfiguration: RTCConfiguration) {
+        
+        // Define media constraints. DtlsSrtpKeyAgreement is required to be true to be able to connect with web browsers.
+        let constraints = RTCMediaConstraints(mandatoryConstraints: nil,
+                                              optionalConstraints: nil /*["DtlsSrtpKeyAgreement":kRTCMediaConstraintsValueTrue]*/)
+        
+        guard let pc = WebRTCClient.factory.peerConnection(with: rtcConfiguration, constraints: constraints, delegate: nil) else {
+            fatalError("Could not create new RTCPeerConnection")
+        }
+        
+        self.peerConnection = pc
+        
+        // Create the data channel
+        if let dataChannel = createDataChannel(peerConnection: pc) {
+            dataChannel.delegate = self
+            self.dataChannel = dataChannel
+        }
+        
+        self.configureAudioSession()
+        self.peerConnection!.delegate = self
+    }
+
+    func stats(_ completionHandler : @escaping RTCStatisticsCompletionHandler) {
+        self.peerConnection?.statistics(completionHandler: completionHandler)
+    }
+    
+    func hasPeerConnnection() -> Bool {
+        if let _ = self.peerConnection {
+            return true
+        } else {
+            return false
+        }
+    }
+    
+    func getRemoteVideoTrack() -> RTCVideoTrack? {
+        return self.remoteVideoTrack
+    }
+    
+    // MARK: Signaling
+    func offer(completion: @escaping (_ sdp: RTCSessionDescription) -> Void) {
+        
+        if let peerConnection = peerConnection {
+            let constrains = RTCMediaConstraints(mandatoryConstraints: self.mediaConstrains,
+                                                 optionalConstraints: nil)
+            
+            peerConnection.offer(for: constrains) { (sdp, error) in
+                guard let sdp = sdp else {
+                    return
+                }
+                
+                peerConnection.setLocalDescription(sdp, completionHandler: { (error) in
+                    completion(sdp)
+                })
+            }
+        }
+        else {
+            debugPrint("Peer connection not setup yet, cannot handle an offer yet.")
+        }
+    }
+    
+    func answer(completion: @escaping (_ sdp: RTCSessionDescription) -> Void)  {
+        
+        if let peerConnection = peerConnection {
+            let constrains = RTCMediaConstraints(mandatoryConstraints: self.mediaConstrains,
+                                                 optionalConstraints: nil)
+            peerConnection.answer(for: constrains) { (sdp, error) in
+                guard let sdp = sdp else {
+                    return
+                }
+                
+                peerConnection.setLocalDescription(sdp, completionHandler: { (error) in
+                    completion(sdp)
+                })
+            }
+        } else {
+            debugPrint("Peer connection not setup yet, cannot handle answer yet.")
+        }
+        
+    }
+    
+    func handleRemoteSdp(remoteSdp: RTCSessionDescription, completion: @escaping (Error?) -> ()) {
+        if let peerConnection = peerConnection {
+            peerConnection.setRemoteDescription(remoteSdp, completionHandler: completion)
+        } else {
+            debugPrint("Peer connection not setup yet, cannot handle setting remote sdp yet.")
+        }
+    }
+    
+    func handleRemoteCandidate(remoteCandidate: RTCIceCandidate, completion: @escaping (Error?) -> ()) {
+        if let peerConnection = peerConnection {
+            peerConnection.add(remoteCandidate, completionHandler: completion)
+        } else {
+            debugPrint("Peer connection not setup yet, cannot handle setting remote candidate yet.")
+        }
+    }
+    
+}
+
+// MARK: Data Channels
+extension WebRTCClient: RTCDataChannelDelegate {
+    
+    private func createDataChannel(peerConnection : RTCPeerConnection) -> RTCDataChannel? {
+        let config = RTCDataChannelConfiguration()
+        guard let dataChannel = peerConnection.dataChannel(forLabel: "iOSDataChannel", configuration: config) else {
+            debugPrint("Warning: Couldn't create data channel.")
+            return nil
+        }
+        return dataChannel
+    }
+    
+    func sendData(_ data: Data) {
+        let buffer = RTCDataBuffer(data: data, isBinary: true)
+        self.dataChannel?.sendData(buffer)
+    }
+    
+    func dataChannelDidChangeState(_ dataChannel: RTCDataChannel) {
+        debugPrint("dataChannel did change state: \(dataChannel.readyState)")
+    }
+    
+    func dataChannel(_ dataChannel: RTCDataChannel, didReceiveMessageWith buffer: RTCDataBuffer) {
+        self.delegate?.webRTCClient(self, didReceiveData: buffer.data)
+    }
+}
+
+// MARK: PeerConnection
+extension WebRTCClient: RTCPeerConnectionDelegate {
+    
+    func peerConnection(_ peerConnection: RTCPeerConnection, didChange stateChanged: RTCSignalingState) {
+        debugPrint("peerConnection new signaling state: \(stateChanged)")
+    }
+    
+    func peerConnection(_ peerConnection: RTCPeerConnection, didAdd stream: RTCMediaStream) {
+        debugPrint("peerConnection did add stream. Id=\(stream.streamId)")
+    }
+    
+    func peerConnection(_ peerConnection: RTCPeerConnection, didRemove stream: RTCMediaStream) {
+        debugPrint("peerConnection did remove stream")
+    }
+    
+    func peerConnectionShouldNegotiate(_ peerConnection: RTCPeerConnection) {
+        debugPrint("peerConnection should negotiate")
+    }
+    
+    func peerConnection(_ peerConnection: RTCPeerConnection, didChange newState: RTCIceConnectionState) {
+        debugPrint("peerConnection new connection state: \(newState)")
+        self.delegate?.webRTCClient(self, didChangeConnectionState: newState)
+    }
+    
+    func peerConnection(_ peerConnection: RTCPeerConnection, didChange newState: RTCIceGatheringState) {
+        debugPrint("peerConnection new gathering state: \(newState)")
+    }
+    
+    func peerConnection(_ peerConnection: RTCPeerConnection, didGenerate candidate: RTCIceCandidate) {
+        self.delegate?.webRTCClient(self, didDiscoverLocalCandidate: candidate)
+    }
+    
+    func peerConnection(_ peerConnection: RTCPeerConnection, didRemove candidates: [RTCIceCandidate]) {
+        debugPrint("peerConnection did remove candidate(s)")
+    }
+    
+    func peerConnection(_ peerConnection: RTCPeerConnection, didOpen dataChannel: RTCDataChannel) {
+        debugPrint("peerConnection did open data channel")
+    }
+    
+    func peerConnection(_ peerConnection: RTCPeerConnection, didAdd rtpReceiver: RTCRtpReceiver, streams mediaStreams: [RTCMediaStream]) {
+        debugPrint("Added rtp receiver id=\(rtpReceiver.receiverId) trackId=\(rtpReceiver.track?.trackId ?? "none")")
+    }
+    
+    func peerConnection(_ peerConnection: RTCPeerConnection, didStartReceivingOn transceiver: RTCRtpTransceiver) {
+        
+        var directionStr = ""
+        switch transceiver.direction {
+        case RTCRtpTransceiverDirection.inactive:
+            directionStr = "inactive"
+        case RTCRtpTransceiverDirection.recvOnly:
+            directionStr = "recvonly"
+        case RTCRtpTransceiverDirection.sendOnly:
+            directionStr = "sendonly"
+        case RTCRtpTransceiverDirection.sendRecv:
+            directionStr = "sendrecv"
+        case RTCRtpTransceiverDirection.stopped:
+            directionStr = "stopped"
+        @unknown default:
+            directionStr = "unknown"
+        }
+        
+        switch transceiver.mediaType {
+        case .video:
+            debugPrint("Started recieving video. Direction= \(directionStr)")
+            transceiver.receiver.track?.isEnabled = true
+            
+            self.remoteVideoTrack = transceiver.receiver.track as? RTCVideoTrack
+            self.delegate?.webRTCClient(self, onStartReceiveVideo: self.remoteVideoTrack!)
+            
+        case .audio:
+            // disable audio sending
+            transceiver.sender.track?.isEnabled = false
+            debugPrint("Started receiving audio. Direction=\(directionStr)")
+        case .data:
+            debugPrint("Started receiving data channel data.")
+        case .unsupported:
+            debugPrint("Start receiving on unknown transceiver type - this is probably an error.")
+        @unknown default:
+            debugPrint("Unknown transceiver type.")
+        }
+    }
+    
+}
+
+// MARK: Audio control
+extension WebRTCClient {
+    
+    private func configureAudioSession() {
+        self.rtcAudioSession.lockForConfiguration()
+        do {
+            try self.rtcAudioSession.setCategory(AVAudioSession.Category.ambient.rawValue) /* Playback only */
+            try self.rtcAudioSession.setMode(AVAudioSession.Mode.default.rawValue) /* Mode to default mode */
+        } catch let error {
+            debugPrint("Error changeing AVAudioSession category: \(error)")
+        }
+        self.rtcAudioSession.unlockForConfiguration()
+    }
+    
+    // Fallback to the default playing device: headphones/bluetooth/ear speaker
+    func speakerOff() {
+        self.audioQueue.async { [weak self] in
+            guard let self = self else {
+                return
+            }
+            
+            self.rtcAudioSession.lockForConfiguration()
+            do {
+                try self.rtcAudioSession.overrideOutputAudioPort(.none)
+                try self.rtcAudioSession.setActive(false)
+            } catch let error {
+                debugPrint("Error setting AVAudioSession category: \(error)")
+            }
+            self.rtcAudioSession.unlockForConfiguration()
+        }
+    }
+    
+    // Force speaker
+    func speakerOn() {
+        self.audioQueue.async { [weak self] in
+            guard let self = self else {
+                return
+            }
+            
+            self.rtcAudioSession.lockForConfiguration()
+            do {
+                try self.rtcAudioSession.overrideOutputAudioPort(.speaker)
+                try self.rtcAudioSession.setActive(true)
+            } catch let error {
+                debugPrint("Couldn't force audio to speaker: \(error)")
+            }
+            self.rtcAudioSession.unlockForConfiguration()
+        }
+    }
+
+}
diff --git a/Engine/Extras/VirtualProduction/LiveLinkVCAM/vcam/WebRTC/Services/WebSocketProvider/NativeWebSocketProvider.swift b/Engine/Extras/VirtualProduction/LiveLinkVCAM/vcam/WebRTC/Services/WebSocketProvider/NativeWebSocketProvider.swift
new file mode 100644
--- /dev/null
+++ b/Engine/Extras/VirtualProduction/LiveLinkVCAM/vcam/WebRTC/Services/WebSocketProvider/NativeWebSocketProvider.swift
@@ -0,0 +1,151 @@
+//
+//  NativeWebSocketProvider.swift
+//  WebRTC-Demo
+//
+//  Created by Luke Bermingham on 7/6/2022.
+//  Copyright © 2022 Stas Seldin. All rights reserved.
+//
+
+import Foundation
+
+@available(iOS 13.0, *)
+class NativeWebSocketProvider : NSObject, WebSocketProvider, URLSessionWebSocketDelegate {
+    
+    var delegate                            : WebSocketProviderDelegate?
+    private var socket                      : URLSessionWebSocketTask!
+    private var timeout                     : TimeInterval!
+    private var url                         : URL!
+    private(set) var isConnected            : Bool                      = false
+    
+    init(url:URL,timeout:TimeInterval) {
+        self.timeout        = timeout
+        self.url            = url
+        super.init()
+    }
+    
+    // do not move create socket to init method because if you want to reconnect it never connect again
+    public func connect() {
+        let configuration                        = URLSessionConfiguration.default
+        let urlSession                           = URLSession(configuration: configuration, delegate: self, delegateQueue: OperationQueue())
+        let urlRequest                           = URLRequest(url: url,timeoutInterval: timeout)
+        socket                                   = urlSession.webSocketTask(with: urlRequest)
+        socket.resume()
+        readMessage()
+    }
+    
+    func close() {
+        let error = NSError(domain: "NativeWebSocketProvider", code: 0, userInfo: [NSLocalizedDescriptionKey : "Manually closed websocket."])
+        closeConnection(error)
+    }
+    
+    func send(data: Data) {
+        socket.send(.data(data)) { error in
+            self.handleError(error)
+        }
+    }
+    
+    func send(text: String) {
+        socket.send(.string(text)) { error in
+            self.handleError(error)
+        }
+    }
+    
+    private func readMessage() {
+        socket.receive { result in
+            switch result {
+            case .failure(_):
+                break
+            case .success(let message):
+                switch message {
+                case .data(let data):
+                    self.delegate?.webSocketDidReceiveData(self, data: data)
+                case .string(let string):
+                    self.delegate?.webSocketDidReceiveData(self, data: string.data(using: .utf8)!)
+                @unknown default:
+                    print("Unimplemented case found in NativeWebSocketProvider.")
+                }
+                self.readMessage()
+            }
+        }
+    }
+    
+    func urlSession(_ session: URLSession, webSocketTask: URLSessionWebSocketTask, didOpenWithProtocol protocol: String?) {
+        isConnected = true
+        delegate?.webSocketDidConnect(self)
+    }
+    
+    func urlSession(_ session: URLSession, webSocketTask: URLSessionWebSocketTask, didCloseWith closeCode: URLSessionWebSocketTask.CloseCode, reason: Data?) {
+        
+        var closeCodeStr : String
+        switch closeCode {
+        case URLSessionWebSocketTask.CloseCode.invalid:
+            closeCodeStr = "invalid"
+        case URLSessionWebSocketTask.CloseCode.normalClosure:
+            closeCodeStr = "normalClosure"
+        case URLSessionWebSocketTask.CloseCode.goingAway:
+            closeCodeStr = "goingAway"
+        case URLSessionWebSocketTask.CloseCode.protocolError:
+            closeCodeStr = "protocolError"
+        case URLSessionWebSocketTask.CloseCode.unsupportedData:
+            closeCodeStr = "unsupportedData"
+        case URLSessionWebSocketTask.CloseCode.noStatusReceived:
+            closeCodeStr = "noStatusReceived"
+        case URLSessionWebSocketTask.CloseCode.abnormalClosure:
+            closeCodeStr = "abnormalClosure"
+        case URLSessionWebSocketTask.CloseCode.invalidFramePayloadData:
+            closeCodeStr = "invalidFramePayloadData"
+        case URLSessionWebSocketTask.CloseCode.policyViolation:
+            closeCodeStr = "policyViolation"
+        case URLSessionWebSocketTask.CloseCode.messageTooBig:
+            closeCodeStr = "messageTooBig"
+        case URLSessionWebSocketTask.CloseCode.mandatoryExtensionMissing:
+            closeCodeStr = "mandatoryExtensionMissing"
+        case URLSessionWebSocketTask.CloseCode.internalServerError:
+            closeCodeStr = "internalServerError"
+        case URLSessionWebSocketTask.CloseCode.tlsHandshakeFailure:
+            closeCodeStr = "tlsHandshakeFailure"
+        @unknown default:
+            closeCodeStr = "unknown close code"
+        }
+            
+        
+        if let reason = reason {
+            let reasonStr = String(decoding: reason, as: UTF8.self)
+            print("Websocket closed (\(closeCodeStr)) because: \(reasonStr)")
+        }
+        else {
+            print("Websocket closed without reason message - (\(closeCodeStr))")
+        }
+        isConnected = false
+    }
+    
+    ///never call delegate?.webSocketDidDisconnect in this method it leads to close next connection
+    func urlSession(_ session: URLSession, didReceive challenge: URLAuthenticationChallenge, completionHandler: @escaping (URLSession.AuthChallengeDisposition, URLCredential?) -> Void) {
+        completionHandler(.useCredential, URLCredential(trust: challenge.protectionSpace.serverTrust!))
+    }
+    
+    func urlSession(_ session: URLSession, task: URLSessionTask, didCompleteWithError error: Error?) {
+        closeConnection(error)
+    }
+    
+    func urlSession(_ session: URLSession, didBecomeInvalidWithError error: Error?) {
+        closeConnection(error)
+    }
+    
+    func closeConnection(_ error : Error?) {
+        socket.cancel(with: .goingAway, reason: nil)
+        isConnected = false
+        delegate?.webSocketDidDisconnect(self, error: error)
+    }
+    
+    /// we need to check if error code is one of the 57 , 60 , 54 timeout no network and internet offline to notify delegate we disconnected from internet
+    private func handleError(_ error: Error?){
+        if let error = error as NSError?{
+            if error.code == 57  || error.code == 60 || error.code == 54{
+                closeConnection(error)
+            }else{
+                delegate?.webSocketReceiveError(self, error: error)
+            }
+        }
+    }
+}
diff --git a/Engine/Extras/VirtualProduction/LiveLinkVCAM/vcam/WebRTC/Services/WebSocketProvider/StarscreamProvider.swift b/Engine/Extras/VirtualProduction/LiveLinkVCAM/vcam/WebRTC/Services/WebSocketProvider/StarscreamProvider.swift
new file mode 100644
--- /dev/null
+++ b/Engine/Extras/VirtualProduction/LiveLinkVCAM/vcam/WebRTC/Services/WebSocketProvider/StarscreamProvider.swift
@@ -0,0 +1,53 @@
+//
+//  StarscreamProvider.swift
+//  WebRTC-Demo
+//
+//  Created by stasel on 15/07/2019.
+//  Copyright © 2019 stasel. All rights reserved.
+//
+
+import Foundation
+import Starscream
+
+class StarscreamWebSocket: WebSocketProvider {
+    
+    var delegate: WebSocketProviderDelegate?
+    private let socket: WebSocket
+    
+    init(url: URL) {
+        self.socket = WebSocket(url: url)
+        self.socket.delegate = self
+    }
+    
+    func connect() {
+        self.socket.connect()
+    }
+    
+    func send(data: Data) {
+        self.socket.write(data: data)
+    }
+    
+    func close() {
+        self.socket.disconnect(forceTimeout: 2, closeCode: 1000 /*normal close*/)
+    }
+}
+
+extension StarscreamWebSocket: Starscream.WebSocketDelegate {
+    func websocketDidConnect(socket: WebSocketClient) {
+        self.delegate?.webSocketDidConnect(self)
+    }
+    
+    func websocketDidDisconnect(socket: WebSocketClient, error: Error?) {
+        self.delegate?.webSocketDidDisconnect(self, error: error)
+    }
+    
+    func websocketDidReceiveMessage(socket: WebSocketClient, text: String) {
+        debugPrint("Warning: Expected to receive data format but received a string. Check the websocket server config.")
+    }
+    
+    func websocketDidReceiveData(socket: WebSocketClient, data: Data) {
+        self.delegate?.webSocketDidReceiveData(self, data: data)
+    }
+    
+    
+}
diff --git a/Engine/Extras/VirtualProduction/LiveLinkVCAM/vcam/WebRTC/Services/WebSocketProvider/WebSocketProvider.swift b/Engine/Extras/VirtualProduction/LiveLinkVCAM/vcam/WebRTC/Services/WebSocketProvider/WebSocketProvider.swift
new file mode 100644
--- /dev/null
+++ b/Engine/Extras/VirtualProduction/LiveLinkVCAM/vcam/WebRTC/Services/WebSocketProvider/WebSocketProvider.swift
@@ -0,0 +1,23 @@
+//
+//  File.swift
+//  WebRTC-Demo
+//
+//  Created by stasel on 15/07/2019.
+//  Copyright © 2019 stasel. All rights reserved.
+//
+
+import Foundation
+
+protocol WebSocketProvider: AnyObject {
+    var delegate: WebSocketProviderDelegate? { get set }
+    func connect()
+    func close()
+    func send(data: Data)
+}
+
+protocol WebSocketProviderDelegate: AnyObject {
+    func webSocketDidConnect(_ webSocket: WebSocketProvider)
+    func webSocketDidDisconnect(_ webSocket: WebSocketProvider, error: Error?)
+    func webSocketReceiveError(_ webSocket: WebSocketProvider, error: Error?)
+    func webSocketDidReceiveData(_ webSocket: WebSocketProvider, data: Data)
+}
diff --git a/Engine/Extras/VirtualProduction/LiveLinkVCAM/vcam/WebRTCStreamingConnection.swift b/Engine/Extras/VirtualProduction/LiveLinkVCAM/vcam/WebRTCStreamingConnection.swift
new file mode 100644
--- /dev/null
+++ b/Engine/Extras/VirtualProduction/LiveLinkVCAM/vcam/WebRTCStreamingConnection.swift
@@ -0,0 +1,372 @@
+//
+//  WebRTCStreamingConnection.swift
+//  vcam
+//
+//  Created by Brian Smith on 6/28/22.
+//  Copyright Epic Games, Inc. All Rights Reserved.
+//
+
+import UIKit
+import WebRTC
+
+class WebRTCStreamingConnection : StreamingConnection {
+    var _url : URL?
+    
+    public var webRTCClient: WebRTCClient?
+    private var webRTCClientState : RTCIceConnectionState?
+    private var signalClient: SignalingClient?
+    private var touchControls: TouchControls?
+    private var webRTCView : WebRTCView?
+    private var rtcVideoTrack : RTCVideoTrack?
+    
+    private var signalingConnected = false
+    private var hasRemoteSdp = false
+    private var hasLocalSdp = false
+    private var remoteCandidateCount = 0
+    private var localCandidateCount = 0
+    private var _statsTimer : Timer?
+    private var _lastBytesReceived : Int?
+    private var _lastBytesReceivedTimestamp : CFTimeInterval?
+
+    override var name : String {
+        get {
+            "WebRTC"
+        }
+    }
+    override var destination : String {
+        get {
+            self._url?.absoluteString ?? ""
+        }
+        set {
+            self._url = URL(string: "ws://\(newValue):80")!
+        }
+    }
+    
+    override var isConnected: Bool {
+        get {
+            return (self.webRTCClientState ?? .disconnected) == .connected
+        }
+    }
+
+    override var renderView: UIView? {
+        didSet {
+            if let rv = renderView {
+                let rtcView = WebRTCView(frame: CGRect(x: 0, y: 0, width: rv.frame.size.width, height: rv.frame.size.height))
+                rv.addSubview(rtcView)
+                rtcView.delegate = self
+                rtcView.layoutToSuperview(.top, .bottom, .left, .right)
+                self.webRTCView = rtcView
+                
+                self.attachVideoTrack()
+            }
+        }
+    }
+    
+    required init(subjectName: String) {
+        super.init(subjectName: subjectName)
+        
+        self.webRTCClient = WebRTCClient()
+        self.webRTCClient?.speakerOff()
+        self.webRTCClient?.delegate = self
+
+        _statsTimer = Timer.scheduledTimer(withTimeInterval: 1, repeats: true, block: { timer in
+            
+            if let webRTC = self.webRTCClient {
+                
+                webRTC.stats({ report in
+                    
+                    for (key,value) in report.statistics {
+                        if key.starts(with: "RTCIceCandidatePair_") {
+                            //Log.info("\(value.values)")
+                            let timestamp = value.timestamp_us
+                            if let bytesReceived = value.values["bytesReceived"] as? Int {
+                                
+                                if bytesReceived > 0 {
+
+                                    var bytesPerSecond : Int?
+                                    if let previousTimestamp = self._lastBytesReceivedTimestamp,
+                                       let previousBytesReceived = self._lastBytesReceived {
+
+                                        bytesPerSecond = Int(Double(bytesReceived - previousBytesReceived) / ((timestamp - previousTimestamp) / 1000000.0))
+                                    }
+
+                                    self._lastBytesReceivedTimestamp = timestamp
+                                    self._lastBytesReceived = bytesReceived
+                                    
+                                    if let bps = bytesPerSecond {
+                                        self.stats = StreamingConnectionStats()
+                                        self.stats?.bytesPerSecond = bps
+                                    }
+                                }
+                            }
+                            
+                            break
+                        }
+                    }
+                })
+            } else {
+                self._lastBytesReceived = nil
+                self._lastBytesReceivedTimestamp = nil
+                self.stats = nil
+            }
+        })
+    }
+    
+    deinit {
+        _statsTimer?.invalidate()
+        webRTCClient = nil
+    }
+
+    override func connect() throws {
+        
+        disconnect()
+        
+        // We will use 3rd party library for websockets.
+        let webSocketProvider: WebSocketProvider
+        
+        if #available(iOS 13.0, *) {
+            webSocketProvider = NativeWebSocketProvider(url: self._url!, timeout: 2.0)
+        } else {
+            webSocketProvider = StarscreamWebSocket(url: self._url!)
+        }
+        
+        self.signalClient = SignalingClient(webSocket: webSocketProvider)
+        self.signalClient?.delegate = self
+        self.signalClient?.connect()
+    }
+    
+    override func reconnect() {
+        disconnect()
+        do {
+            try connect()
+        } catch {
+            Log.error(error.localizedDescription)
+        }
+    }
+    
+    override func disconnect() {
+        self.signalClient?.close()
+        self.signalClient = nil
+        signalingConnected = false
+        hasLocalSdp = false
+        hasRemoteSdp = false
+        remoteCandidateCount = 0
+        localCandidateCount = 0
+    }
+    
+    override func sendTransform(_ transform: simd_float4x4, atTime time: Double) {
+    
+        guard let client = webRTCClient else { return }
+    
+        
+        // convert the transform to UE space
+        // adapted from AppleARKitConversion.h
+        let rawRotation = simd_quaternion(transform)
+        let ueRotation = simd_quaternion(-rawRotation.vector.z, rawRotation.vector.x, rawRotation.vector.y, -rawRotation.vector.w)
+        var ueTranslationMatrix = simd_float4x4()
+        ueTranslationMatrix.columns.3 = simd_float4(x: -transform.columns.3.z, y: transform.columns.3.x, z: transform.columns.3.y, w: 1.0) * 100.0 // ue units
+
+        let ueTransform = simd_float4x4(ueRotation) * ueTranslationMatrix
+        
+        var bytes: [UInt8] = []
+
+        // Write message type using 1 byte
+        bytes.append(PixelStreamingToStreamerMessage.Transform.rawValue)
+
+        // Write 4x4 transform each element to get 4 bytes e.g. float -> [UInt8]
+        bytes.append(contentsOf: ueTransform.columns.0.x.toBytes())
+        bytes.append(contentsOf: ueTransform.columns.0.y.toBytes())
+        bytes.append(contentsOf: ueTransform.columns.0.z.toBytes())
+        bytes.append(contentsOf: ueTransform.columns.0.w.toBytes())
+
+        bytes.append(contentsOf: ueTransform.columns.1.x.toBytes())
+        bytes.append(contentsOf: ueTransform.columns.1.y.toBytes())
+        bytes.append(contentsOf: ueTransform.columns.1.z.toBytes())
+        bytes.append(contentsOf: ueTransform.columns.1.w.toBytes())
+
+        bytes.append(contentsOf: ueTransform.columns.2.x.toBytes())
+        bytes.append(contentsOf: ueTransform.columns.2.y.toBytes())
+        bytes.append(contentsOf: ueTransform.columns.2.z.toBytes())
+        bytes.append(contentsOf: ueTransform.columns.2.w.toBytes())
+
+        bytes.append(contentsOf: ueTransform.columns.3.x.toBytes())
+        bytes.append(contentsOf: ueTransform.columns.3.y.toBytes())
+        bytes.append(contentsOf: ueTransform.columns.3.z.toBytes())
+        bytes.append(contentsOf: ueTransform.columns.3.w.toBytes())
+
+        // Write timestamp 8 bytes
+        bytes.append(contentsOf: time.toBytes())
+
+        // Send the transform + timestamp across
+        client.sendData(Data(bytes))
+    }
+    
+    func attachVideoTrack() {
+        if let webRTC = webRTCClient, let view = self.webRTCView, let track = self.rtcVideoTrack {
+            self.touchControls = TouchControls(webRTC, touchView: view)
+            view.attachVideoTrack(track: track)
+            view.attachTouchDelegate(delegate: self.touchControls!)
+        }
+    }
+}
+extension WebRTCStreamingConnection : WebRTCViewDelegate {
+    
+    func webRTCView(_ view: WebRTCView, didChangeVideoSize size: CGSize) {
+        self.videoSize = size
+    }
+}
+
+extension WebRTCStreamingConnection: SignalClientDelegate {
+    
+    func signalClientDidConnect(_ signalClient: SignalingClient) {
+        self.signalingConnected = true
+        Log.info("Connected to signaling server")
+    }
+    
+    func signalClientDidDisconnect(_ signalClient: SignalingClient, error: Error?) {
+        self.signalingConnected = false
+        Log.info("Disconnected from signaling server")
+    }
+    
+    func signalClientDidReceiveError(_ signalClient: SignalingClient, error: Error?) {
+        if error != nil {
+            Log.error("Signalling got error: \(error!)")
+        }
+    }
+    
+    func signalClient(_ signalClient: SignalingClient, didReceiveConfig config: RTCConfiguration) {
+        Log.info("Received peer connection configuration - ICE servers: \(config.iceServers)")
+        self.webRTCClient?.setupPeerConnection(rtcConfiguration: config)
+    }
+    
+    func signalClient(_ signalClient: SignalingClient, didReceiveRemoteSdp sdp: RTCSessionDescription) {
+        
+        var sdpTypeStr : String = ""
+        switch sdp.type {
+        case RTCSdpType.answer:
+            sdpTypeStr = "answer"
+        case RTCSdpType.offer:
+            sdpTypeStr = "offer"
+        case RTCSdpType.prAnswer:
+            sdpTypeStr = "prAnswer"
+        case RTCSdpType.rollback:
+            sdpTypeStr = "rollback"
+        @unknown default:
+            sdpTypeStr = "unknown"
+        }
+        
+        Log.info("Received remote sdp. Type=\(sdpTypeStr)")
+        Log.info(sdp.sdp)
+        
+        if self.webRTCClient!.hasPeerConnnection() {
+            self.webRTCClient!.handleRemoteSdp(remoteSdp: sdp) { (error) in
+                self.hasRemoteSdp = true
+                
+                // If we get an offer from the streamer we send an answer back
+                if sdp.type == RTCSdpType.offer {
+                    self.signalClientSendAnswer(signalClient)
+                }
+                else {
+                    Log.debug("We only support replying to offer, but we got \(sdpTypeStr)")
+                }
+            }
+        } else {
+            Log.debug("WebRTC peer connection not setup yet - cannot handle remote sdp.")
+        }
+    }
+    
+    func signalClient(_ signalClient: SignalingClient, didReceiveCandidate candidate: RTCIceCandidate) {
+        Log.info("Received remote candidate - \(candidate.sdp)")
+        
+        if self.webRTCClient!.hasPeerConnnection() {
+            self.webRTCClient!.handleRemoteCandidate(remoteCandidate: candidate) { error in
+                self.remoteCandidateCount += 1
+            }
+        } else {
+            Log.debug("WebRTC peer connection not setup yet - cannot handle remote candidate")
+        }
+    }
+    
+    func signalClientSendAnswer(_ signalClient: SignalingClient){
+        if self.webRTCClient!.hasPeerConnnection() {
+            Log.info("Sending answer sdp")
+            self.webRTCClient!.answer { (localSdp) in
+                Log.info(localSdp.sdp)
+                self.hasLocalSdp = true
+                signalClient.send(sdp: localSdp)
+            }
+        } else {
+            Log.debug("WebRTC peer connection not setup yet - cannot handle sending answer.")
+        }
+    }
+    
+}
+
+extension WebRTCStreamingConnection: WebRTCClientDelegate {
+    
+    func webRTCClient(_ client: WebRTCClient, onStartReceiveVideo video: RTCVideoTrack) {
+        self.rtcVideoTrack = video
+        self.attachVideoTrack()
+        self.delegate?.streamingConnectionDidConnect(self)
+    }
+    
+    func webRTCClient(_ client: WebRTCClient, didDiscoverLocalCandidate candidate: RTCIceCandidate) {
+        Log.info("discovered local candidate")
+        self.localCandidateCount += 1
+        self.signalClient?.send(candidate: candidate)
+    }
+    
+    func webRTCClient(_ client: WebRTCClient, didChangeConnectionState state: RTCIceConnectionState) {
+        self.webRTCClientState = state
+        switch state {
+        case .connected:
+            
+            // WebRTC connection is "connected" so try to show remote video track if we got one some time during the connection
+            // IMPORTANT: Even though we do the same in onStartReceiveVideo above, this one actually makes the video show up due
+            // peer connection needing to be connected before video can reasonably be displayed - the other callback is more useful
+            // if video is added later on into the call.
+            self.rtcVideoTrack = self.webRTCClient?.getRemoteVideoTrack()
+            self.attachVideoTrack()
+            
+        case .disconnected:
+            self.delegate?.streamingConnection(self, didDisconnectWithError: nil)
+        default:
+            break
+        }
+        Log.info("WebRTC status: \(state.description.capitalized)")
+    }
+    
+    func webRTCClient(_ client: WebRTCClient, didReceiveData data: Data) {
+        
+        if(data.count > 0) {
+            let payloadTypeInt : UInt8 = data[0]
+            if let payloadType = PixelStreamingToClientMessage(rawValue: payloadTypeInt) {
+                switch payloadType {
+                case .VideoEncoderAvgQP:
+                    let qp : String? = String(data: data.dropFirst(), encoding: .utf16LittleEndian)
+                    Log.info("Quality = \(qp ?? "N/A")")
+                case .QualityControlOwnership:
+                    fallthrough
+                case .Response:
+                    fallthrough
+                case .Command:
+                    fallthrough
+                case .FreezeFrame:
+                    fallthrough
+                case .UnfreezeFrame:
+                    fallthrough
+                case .LatencyTest:
+                    fallthrough
+                case .InitialSettings:
+                    fallthrough
+                case .FileExtension:
+                    fallthrough
+                case .FileMimeType:
+                    fallthrough
+                case .FileContents:
+                    print("Skipping payload type \(payloadType) - we have no implementation for it.")
+                }
+            }
+        }
+    }
+}
diff --git a/Engine/Extras/VirtualProduction/LiveLinkVCAM/vcam/WebRTCView.swift b/Engine/Extras/VirtualProduction/LiveLinkVCAM/vcam/WebRTCView.swift
new file mode 100644
--- /dev/null
+++ b/Engine/Extras/VirtualProduction/LiveLinkVCAM/vcam/WebRTCView.swift
@@ -0,0 +1,134 @@
+//
+//  WebRTCView.swift
+//  WebRTC-Demo
+//
+//  Created by Luke Bermingham on 14/6/2022.
+
+import Foundation
+import UIKit
+import WebRTC
+
+protocol WebRTCViewDelegate {
+    func webRTCView(_ view : WebRTCView, didChangeVideoSize size : CGSize)
+}
+
+final class WebRTCView: UIView, RTCVideoViewDelegate {
+    
+    var delegate : WebRTCViewDelegate?
+    
+    let videoView = RTCMTLVideoViewWithTouch(frame: .zero)
+    var videoSize = CGSize.zero
+
+    override init(frame: CGRect) {
+        super.init(frame: frame)
+        
+        self.videoView.videoContentMode = .scaleAspectFit
+        self.videoView.isEnabled = true
+        self.videoView.isHidden = false
+        self.videoView.backgroundColor = UIColor.lightGray
+        
+        videoView.delegate = self
+        addSubview(videoView)
+    }
+
+    required init?(coder: NSCoder) {
+        super.init(coder: coder)
+        
+        videoView.delegate = self
+        addSubview(videoView)
+    }
+
+    func videoView(_ videoView: RTCVideoRenderer, didChangeVideoSize size: CGSize) {
+        self.videoSize = size
+        self.videoView.touchDelegate?.onVideoChangedSize(size: size)
+        setNeedsLayout()
+        
+        self.delegate?.webRTCView(self, didChangeVideoSize: size)
+    }
+    
+    func attachVideoTrack(track: RTCVideoTrack) {
+        track.remove(self.videoView)
+        track.add(self.videoView)
+        track.isEnabled = true
+        self.videoView.renderFrame(nil)
+        let trackState : RTCMediaStreamTrackState = track.readyState
+        if trackState == RTCMediaStreamTrackState.live {
+            print("Video track is live | id=\(track.trackId) | kind=\(track.kind)")
+        }
+    }
+    
+    func attachTouchDelegate(delegate: TouchDelegate) {
+        self.videoView.touchDelegate = delegate
+    }
+
+    override func layoutSubviews() {
+        super.layoutSubviews()
+        guard videoSize.width > 0 && videoSize.height > 0 else {
+            videoView.frame = bounds
+            return
+        }
+
+        var videoFrame = AVMakeRect(aspectRatio: videoSize, insideRect: bounds)
+        let scale = videoFrame.size.aspectFitScale(in: bounds.size)
+        videoFrame.size.width = videoFrame.size.width * CGFloat(scale)
+        videoFrame.size.height = videoFrame.size.height * CGFloat(scale)
+
+        videoView.frame = videoFrame
+        videoView.center = CGPoint(x: bounds.midX, y: bounds.midY)
+    }
+    
+}
+
+extension CGSize {
+    func aspectFitScale(in container: CGSize) -> CGFloat {
+
+        if height <= container.height && width > container.width {
+            return container.width / width
+        }
+        if height > container.height && width > container.width {
+            return min(container.width / width, container.height / height)
+        }
+        if height > container.height && width <= container.width {
+            return container.height / height
+        }
+        if height <= container.height && width <= container.width {
+            return min(container.width / width, container.height / height)
+        }
+        return 1.0
+    }
+}
+
+class RTCMTLVideoViewWithTouch : RTCMTLVideoView {
+    
+    weak var touchDelegate : TouchDelegate?
+    
+    override init(frame: CGRect) {
+        super.init(frame: frame)
+        self.isMultipleTouchEnabled = true
+    }
+    
+    required init?(coder aDecoder: NSCoder) {
+        super.init(coder: aDecoder)
+        isMultipleTouchEnabled = true
+    }
+    
+    open override func touchesBegan(_ touches: Set<UITouch>, with event: UIEvent?) {
+        super.touchesBegan(touches, with: event)
+        self.touchDelegate?.touchesBegan(touches)
+    }
+    
+    open override func touchesMoved(_ touches: Set<UITouch>, with event: UIEvent?) {
+        super.touchesMoved(touches, with: event)
+        self.touchDelegate?.touchesMoved(touches)
+    }
+    
+    open override func touchesEnded(_ touches: Set<UITouch>, with event: UIEvent?) {
+        super.touchesEnded(touches, with: event)
+        self.touchDelegate?.touchesEnded(touches)
+    }
+    
+    open override func touchesCancelled(_ touches: Set<UITouch>, with event: UIEvent?) {
+        super.touchesCancelled(touches, with: event)
+        self.touchDelegate?.touchesCancelled(touches)
+    }
+}
