Corrected the conversion of ARKit transform to the form expected by Unreal

diff --git a/Engine/Extras/VirtualProduction/LiveLinkVCAM/vcam/WebRTCStreamingConnection.swift b/Engine/Extras/VirtualProduction/LiveLinkVCAM/vcam/WebRTCStreamingConnection.swift
--- a/Engine/Extras/VirtualProduction/LiveLinkVCAM/vcam/WebRTCStreamingConnection.swift
+++ b/Engine/Extras/VirtualProduction/LiveLinkVCAM/vcam/WebRTCStreamingConnection.swift
@@ -163,10 +163,9 @@ class WebRTCStreamingConnection : StreamingConnection {
         // adapted from AppleARKitConversion.h
         let rawRotation = simd_quaternion(transform)
         let ueRotation = simd_quaternion(-rawRotation.vector.z, rawRotation.vector.x, rawRotation.vector.y, -rawRotation.vector.w)
-        var ueTranslationMatrix = simd_float4x4()
-        ueTranslationMatrix.columns.3 = simd_float4(x: -transform.columns.3.z, y: transform.columns.3.x, z: transform.columns.3.y, w: 1.0) * 100.0 // ue units
-
-        let ueTransform = simd_float4x4(ueRotation) * ueTranslationMatrix
+        var ueTransform = simd_float4x4(ueRotation)
+        ueTransform.columns.3 = simd_float4(x: -transform.columns.3.z, y: transform.columns.3.x, z: transform.columns.3.y, w: 1.0) * 100.0 // ue units
+        ueTransform.columns.3.w = 1.0 // Force Unit scale
         
         var bytes: [UInt8] = []
 
