Fixed a bug where the disconnect wouldn't work when the signalling server disconnected first (before the webRTC client)

diff --git a/Engine/Extras/VirtualProduction/LiveLinkVCAM/vcam/WebRTCStreamingConnection.swift b/Engine/Extras/VirtualProduction/LiveLinkVCAM/vcam/WebRTCStreamingConnection.swift
--- a/Engine/Extras/VirtualProduction/LiveLinkVCAM/vcam/WebRTCStreamingConnection.swift
+++ b/Engine/Extras/VirtualProduction/LiveLinkVCAM/vcam/WebRTCStreamingConnection.swift
@@ -225,6 +225,7 @@ extension WebRTCStreamingConnection: SignalClientDelegate {
     func signalClientDidDisconnect(_ signalClient: SignalingClient, error: Error?) {
         self.signalingConnected = false
         Log.info("Disconnected from signaling server")
+        self.delegate?.streamingConnection(self, didDisconnectWithError: error)
     }
     
     func signalClientDidReceiveError(_ signalClient: SignalingClient, error: Error?) {
