Added timeout when disconnecting. After 0.5 seconds if the disconnect hasn't initiated from the UE side, we force the disconnect and return to the start screen.

diff --git a/Engine/Extras/VirtualProduction/LiveLinkVCAM/vcam/VideoViewController+StreamingConnectionDelegate.swift b/Engine/Extras/VirtualProduction/LiveLinkVCAM/vcam/VideoViewController+StreamingConnectionDelegate.swift
--- a/Engine/Extras/VirtualProduction/LiveLinkVCAM/vcam/VideoViewController+StreamingConnectionDelegate.swift
+++ b/Engine/Extras/VirtualProduction/LiveLinkVCAM/vcam/VideoViewController+StreamingConnectionDelegate.swift
@@ -27,9 +27,7 @@ extension VideoViewController : StreamingConnectionDelegate {
         DispatchQueue.main.async {
 
             if self.dismissOnDisconnect {
-                self.streamingConnection?.delegate = nil
-                self.streamingConnection = nil
-                self.presentingViewController?.dismiss(animated: true, completion: nil)
+                self.forceDisconnectAndDismiss()
             } else {
                 self.showReconnecting(true, animated: true)
             }
diff --git a/Engine/Extras/VirtualProduction/LiveLinkVCAM/vcam/VideoViewController.swift b/Engine/Extras/VirtualProduction/LiveLinkVCAM/vcam/VideoViewController.swift
--- a/Engine/Extras/VirtualProduction/LiveLinkVCAM/vcam/VideoViewController.swift
+++ b/Engine/Extras/VirtualProduction/LiveLinkVCAM/vcam/VideoViewController.swift
@@ -36,6 +36,7 @@ class VideoViewController : BaseViewController {
 
     weak var liveLink : LiveLinkProvider?
     var dismissOnDisconnect = false
+    var forceDisconnectTimer : Timer?
     
     var statsTimer : Timer?
     
@@ -237,8 +238,7 @@ class VideoViewController : BaseViewController {
             dismissOnDisconnect = true
             disconnect()
         } else {
-            disconnect()
-            self.presentingViewController?.dismiss(animated: true, completion: nil)
+            self.forceDisconnectAndDismiss()
         }
     }
     
@@ -247,8 +247,19 @@ class VideoViewController : BaseViewController {
     }
     
     func disconnect() {
+        self.forceDisconnectTimer = Timer.scheduledTimer(withTimeInterval: 0.5, repeats: false, block: { timer in
+            Log.info("Forcing disconnect due to timeout.")
+            self.forceDisconnectAndDismiss()
+        })
         self.streamingConnection?.disconnect()
-        
+    }
+    
+    func forceDisconnectAndDismiss() {
+        self.forceDisconnectTimer?.invalidate()
+        self.forceDisconnectTimer = nil
+        self.streamingConnection?.delegate = nil
+        self.streamingConnection = nil
+        self.presentingViewController?.dismiss(animated: true, completion: nil)
     }
     
     func updateStreamingStats() {
