Fixed broken layout when changing the resolution of a remote session.

diff --git a/Engine/Extras/VirtualProduction/LiveLinkVCAM/vcam/RemoteSessionView.swift b/Engine/Extras/VirtualProduction/LiveLinkVCAM/vcam/RemoteSessionView.swift
--- a/Engine/Extras/VirtualProduction/LiveLinkVCAM/vcam/RemoteSessionView.swift
+++ b/Engine/Extras/VirtualProduction/LiveLinkVCAM/vcam/RemoteSessionView.swift
@@ -32,7 +32,7 @@ class RemoteSessionView : UIView {
                 let thisSize = CGSize(width: CVPixelBufferGetWidth(pb), height: CVPixelBufferGetHeight(pb))
                 if _aspectRatioConstraint == nil || _aspectRatioConstraint?.multiplier != thisSize.width / thisSize.height {
                     if let constraint = _aspectRatioConstraint {
-                        self.removeConstraint(constraint)
+                        self.superview?.removeConstraint(constraint)
                     }
                     _aspectRatioConstraint = layout(.width, to: .height, of: self, ratio: (thisSize.width / thisSize.height), priority: .required)
                 }
