Live Link VCAM (iOS) v1.3.0 build 5 : Re-enabled HW controller support for both Remote Session & Pixel Streaming. Needs a few updates on the UE side as triggers don't yet work. The right thumbstick Y axis also needs work.

diff --git a/Engine/Extras/VirtualProduction/LiveLinkVCAM/vcam/RemoteSessionStreamingConnection.swift b/Engine/Extras/VirtualProduction/LiveLinkVCAM/vcam/RemoteSessionStreamingConnection.swift
--- a/Engine/Extras/VirtualProduction/LiveLinkVCAM/vcam/RemoteSessionStreamingConnection.swift
+++ b/Engine/Extras/VirtualProduction/LiveLinkVCAM/vcam/RemoteSessionStreamingConnection.swift
@@ -150,6 +150,21 @@ class RemoteSessionStreamingConnection : StreamingConnection  {
         self._liveLink?.updateSubject(AppSettings.shared.liveLinkSubjectName, withTransform: transform, atTime: time)
     }
     
+    override func sendControllerAnalog(_ type : StreamingConnectionControllerInputType, controllerIndex : UInt8, value : Float) {
+        
+        Log.info("sendControllerAnalog \(type.rawValue) \(value)")
+        self.oscConnection?.send(.controllerAnalog, arguments: [ OSCArgument.blob(OSCUtility.ueControllerAnalogData(key: type.rawValue, controller: Int(controllerIndex), value: value)) ])
+    }
+    
+    override func sendControllerButtonPressed(_ type : StreamingConnectionControllerInputType, controllerIndex : UInt8, isRepeat : Bool) {
+        self.oscConnection?.send(.controllerButtonPressed, arguments: [ OSCArgument.blob(OSCUtility.ueControllerButtonData(key: type.rawValue, controller: Int(controllerIndex), isRepeat: isRepeat)) ] )
+    }
+
+   override func sendControllerButtonReleased(_ type : StreamingConnectionControllerInputType, controllerIndex : UInt8) {
+       self.oscConnection?.send(.controllerButtonReleased, arguments: [ OSCArgument.blob(OSCUtility.ueControllerButtonData(key: type.rawValue, controller: Int(controllerIndex), isRepeat: false)) ] )
+    }
+
+    
     func restartLiveLink() {
 
         // stop the provider & restart livelink here
diff --git a/Engine/Extras/VirtualProduction/LiveLinkVCAM/vcam/StreamingConnection.swift b/Engine/Extras/VirtualProduction/LiveLinkVCAM/vcam/StreamingConnection.swift
--- a/Engine/Extras/VirtualProduction/LiveLinkVCAM/vcam/StreamingConnection.swift
+++ b/Engine/Extras/VirtualProduction/LiveLinkVCAM/vcam/StreamingConnection.swift
@@ -21,6 +21,30 @@ enum StreamingConnectionTouchType : String {
     case ended = "Ended"
 }
 
+enum StreamingConnectionControllerInputType : String {
+    case thumbstickLeftX = "Gamepad_LeftX"
+    case thumbstickLeftY = "Gamepad_LeftY"
+    case thumbstickRightX = "Gamepad_RightX"
+    case thumbstickRightY = "Gamepad_RightY"
+    case thumbstickLeftButton = "Gamepad_LeftThumbstick"
+    case thumbstickRightButton = "Gamepad_RightThumbstick"
+    case faceButtonBottom = "Gamepad_FaceButton_Bottom"
+    case faceButtonRight = "Gamepad_FaceButton_Right"
+    case faceButtonLeft = "Gamepad_FaceButton_Left"
+    case faceButtonTop = "Gamepad_FaceButton_Top"
+    case shoulderButtonLeft = "Gamepad_LeftShoulder"
+    case shoulderButtonRight = "Gamepad_RightShoulder"
+    case triggerButtonLeft = "Gamepad_LeftTrigger"
+    case triggerButtonRight = "Gamepad_RightTrigger"
+    case dpadUp = "Gamepad_DPad_Up"
+    case dpadDown = "Gamepad_DPad_Down"
+    case dpadLeft = "Gamepad_DPad_Left"
+    case dpadRight = "Gamepad_DPad_Right"
+    case specialButtonLeft = "Gamepad_Special_Left"
+    case specialButtonRight = "Gamepad_Special_Right"
+}
+
+
 class StreamingConnectionStats {
     var bytesPerSecond : Int?
     var framesPerSecond : Float?
@@ -91,6 +115,18 @@ class StreamingConnection : NSObject {
         assertionFailure("not implemented")
     }
     
+    func sendControllerAnalog(_ type :StreamingConnectionControllerInputType, controllerIndex : UInt8, value : Float) {
+        assertionFailure("not implemented")
+    }
+    
+    func sendControllerButtonPressed(_ type : StreamingConnectionControllerInputType, controllerIndex : UInt8, isRepeat : Bool) {
+        assertionFailure("not implemented")
+    }
+
+    func sendControllerButtonReleased(_ type : StreamingConnectionControllerInputType, controllerIndex : UInt8) {
+        assertionFailure("not implemented")
+    }
+
     func setRenderView(_ view : UIView) {
         assertionFailure("not implemented")
     }
diff --git a/Engine/Extras/VirtualProduction/LiveLinkVCAM/vcam/VideoViewController+ARSessionDelegate.swift b/Engine/Extras/VirtualProduction/LiveLinkVCAM/vcam/VideoViewController+ARSessionDelegate.swift
--- a/Engine/Extras/VirtualProduction/LiveLinkVCAM/vcam/VideoViewController+ARSessionDelegate.swift
+++ b/Engine/Extras/VirtualProduction/LiveLinkVCAM/vcam/VideoViewController+ARSessionDelegate.swift
@@ -41,6 +41,6 @@ extension VideoViewController : ARSessionDelegate {
         self.streamingConnection?.sendTransform(frame.camera.transform, atTime: Timecode.create().toTimeInterval())
         
         // update controller
-        self.sendControllerUpdate()
+        self.sendControllerThumbstickUpdate()
     }
 }
diff --git a/Engine/Extras/VirtualProduction/LiveLinkVCAM/vcam/VideoViewController.swift b/Engine/Extras/VirtualProduction/LiveLinkVCAM/vcam/VideoViewController.swift
--- a/Engine/Extras/VirtualProduction/LiveLinkVCAM/vcam/VideoViewController.swift
+++ b/Engine/Extras/VirtualProduction/LiveLinkVCAM/vcam/VideoViewController.swift
@@ -44,8 +44,6 @@ class VideoViewController : BaseViewController {
     weak var gameController : GCController? {
         didSet {
             if let gc = gameController {
-                //gc.extendedGamepad?.leftThumbstick.valueChangedHandler = self.directionalPadValueChanged
-                //gc.extendedGamepad?.rightThumbstick.valueChangedHandler = self.directionalPadValueChanged
                 gc.extendedGamepad?.dpad.up.pressedChangedHandler = self.buttonValueChanged
                 gc.extendedGamepad?.dpad.down.pressedChangedHandler = self.buttonValueChanged
                 gc.extendedGamepad?.dpad.left.pressedChangedHandler = self.buttonValueChanged
@@ -359,118 +357,89 @@ class VideoViewController : BaseViewController {
         self.headerViewY = self.headerView.layer.presentation()?.frame.minY ?? 0
     }
     
-    func sendThumbstickUpdate(key : String, controller : Int, oldValue : Float?, newValue : Float) {
-
-        assertionFailure()
-        /*
-        guard let osc = oscConnection else { return }
-
-        if (newValue != 0.0) || ((oldValue ?? 0.0) != newValue) {
-            osc.send(.controllerAnalog, arguments: [ OSCArgument.blob(OSCUtility.ueControllerAnalogData(key: key, controller: controller, value: newValue)) ])
-        }
-         */
-    }
-    
-    func sendControllerUpdate() {
+    func sendControllerThumbstickUpdate() {
         
+        guard let sc = streamingConnection else { return }
         guard let gc = gameController else { return }
-        assertionFailure()
         
-        /*
-        guard let _ = oscConnection else { return }
-
         let snapshot = gc.capture()
         guard let gp = snapshot.extendedGamepad else { return }
 
         let controllerIndex = gc.playerIndex.rawValue
-        
-        sendThumbstickUpdate(key: "Gamepad_LeftX", controller: controllerIndex, oldValue : self.gameControllerSnapshot?.extendedGamepad?.leftThumbstick.xAxis.value, newValue :gp.leftThumbstick.xAxis.value)
-        sendThumbstickUpdate(key: "Gamepad_LeftY", controller: controllerIndex, oldValue : self.gameControllerSnapshot?.extendedGamepad?.leftThumbstick.yAxis.value, newValue :gp.leftThumbstick.yAxis.value)
-        sendThumbstickUpdate(key: "Gamepad_RightX", controller: controllerIndex, oldValue : self.gameControllerSnapshot?.extendedGamepad?.rightThumbstick.xAxis.value, newValue :gp.rightThumbstick.xAxis.value)
-        sendThumbstickUpdate(key: "Gamepad_RightY", controller: controllerIndex, oldValue : self.gameControllerSnapshot?.extendedGamepad?.rightThumbstick.yAxis.value, newValue :gp.rightThumbstick.yAxis.value)
-        
-        self.gameControllerSnapshot = snapshot
-         */
-    }
-
-    func directionalPadValueChanged(directionalPad : GCControllerDirectionPad, xValue : Float, yValue : Float) {
-        assertionFailure()
-        /*
-        guard let osc = oscConnection else { return }
-        guard let gc = gameController else { return }
-        guard let gp = gc.extendedGamepad else { return }
 
-        let controllerIndex = gc.playerIndex.rawValue
+        // dict of type -> oldValue, newValue
+        let inputs : [StreamingConnectionControllerInputType : ( oldValue: Float?, newValue : Float)]  = [
+            .thumbstickLeftX : ( self.gameControllerSnapshot?.extendedGamepad?.leftThumbstick.xAxis.value,  gp.leftThumbstick.xAxis.value),
+            .thumbstickLeftY : ( self.gameControllerSnapshot?.extendedGamepad?.leftThumbstick.yAxis.value,  gp.leftThumbstick.yAxis.value),
+            .thumbstickRightX : ( self.gameControllerSnapshot?.extendedGamepad?.rightThumbstick.xAxis.value,  gp.rightThumbstick.xAxis.value),
+            .thumbstickRightY : ( self.gameControllerSnapshot?.extendedGamepad?.rightThumbstick.yAxis.value,  gp.rightThumbstick.yAxis.value),
+        ]
         
-        switch directionalPad {
-        case gp.leftThumbstick:
-            osc.send(.controllerAnalog, arguments: [ OSCArgument.blob(OSCUtility.ueControllerAnalogData(key: "Gamepad_LeftX", controller: controllerIndex, value: xValue)) ])
-            osc.send(.controllerAnalog, arguments: [ OSCArgument.blob(OSCUtility.ueControllerAnalogData(key: "Gamepad_LeftY", controller: controllerIndex, value: yValue)) ])
-        case gp.rightThumbstick:
-            osc.send(.controllerAnalog, arguments: [ OSCArgument.blob(OSCUtility.ueControllerAnalogData(key: "Gamepad_RightX", controller: controllerIndex, value: xValue)) ])
-            osc.send(.controllerAnalog, arguments: [ OSCArgument.blob(OSCUtility.ueControllerAnalogData(key: "Gamepad_RightY", controller: controllerIndex, value: yValue)) ])
-        default:
-            Log.warning("directionalPadValueChanged() encounted unsupported '\(directionalPad.localizedName ?? "UNK")'")
-            break
+        for input in inputs {
+            
+            if (input.value.newValue != 0.0) || ((input.value.oldValue ?? 0.0) != input.value.newValue) {
+                sc.sendControllerAnalog(input.key, controllerIndex: UInt8(controllerIndex), value: input.value.newValue)
+            }
         }
-         */
+
+        self.gameControllerSnapshot = snapshot
     }
-    
+
     func buttonValueChanged(button : GCControllerButtonInput, value : Float, pressed : Bool) {
 
-        assertionFailure()
-        /*
-        guard let osc = oscConnection else { return }
+        
+        guard let sc = streamingConnection else { return }
         guard let gc = gameController else { return }
         guard let gp = gc.extendedGamepad else { return }
 
-        let address = button.isPressed ? OSCAddressPattern.controllerButtonPressed : OSCAddressPattern.controllerButtonReleased
         let controllerIndex = gc.playerIndex.rawValue
         let isRepeat = false
+        
+        var inputType : StreamingConnectionControllerInputType!
 
         switch button {
         case gp.leftThumbstickButton:
-            osc.send(address, arguments: [ OSCArgument.blob(OSCUtility.ueControllerButtonData(key: "Gamepad_LeftThumbstick", controller: controllerIndex, isRepeat: isRepeat)) ])
+            inputType = .thumbstickLeftButton
         case gp.rightThumbstickButton:
-            osc.send(address, arguments: [ OSCArgument.blob(OSCUtility.ueControllerButtonData(key: "Gamepad_RightThumbstick", controller: controllerIndex, isRepeat: isRepeat)) ])
-
+            inputType = .thumbstickRightButton
         case gp.buttonA:
-            osc.send(address, arguments: [ OSCArgument.blob(OSCUtility.ueControllerButtonData(key: "Gamepad_FaceButton_Bottom", controller: controllerIndex, isRepeat: isRepeat)) ])
+            inputType = .faceButtonBottom
         case gp.buttonB:
-            osc.send(address, arguments: [ OSCArgument.blob(OSCUtility.ueControllerButtonData(key: "Gamepad_FaceButton_Right", controller: controllerIndex, isRepeat: isRepeat)) ])
+            inputType = .faceButtonRight
         case gp.buttonX:
-            osc.send(address, arguments: [ OSCArgument.blob(OSCUtility.ueControllerButtonData(key: "Gamepad_FaceButton_Left", controller: controllerIndex, isRepeat: isRepeat)) ])
+            inputType = .faceButtonLeft
         case gp.buttonY:
-            osc.send(address, arguments: [ OSCArgument.blob(OSCUtility.ueControllerButtonData(key: "Gamepad_FaceButton_Top", controller: controllerIndex, isRepeat: isRepeat)) ])
-
+            inputType = .faceButtonTop
         case gp.leftShoulder:
-            osc.send(address, arguments: [ OSCArgument.blob(OSCUtility.ueControllerButtonData(key: "Gamepad_LeftShoulder", controller: controllerIndex, isRepeat: isRepeat)) ])
+            inputType = .shoulderButtonLeft
         case gp.rightShoulder:
-            osc.send(address, arguments: [ OSCArgument.blob(OSCUtility.ueControllerButtonData(key: "Gamepad_RightShoulder", controller: controllerIndex, isRepeat: isRepeat)) ])
-
+            inputType = .shoulderButtonRight
         case gp.leftTrigger:
-            osc.send(address, arguments: [ OSCArgument.blob(OSCUtility.ueControllerButtonData(key: "Gamepad_LeftTrigger", controller: controllerIndex, isRepeat: isRepeat)) ])
+            inputType = .triggerButtonLeft
         case gp.rightTrigger:
-            osc.send(address, arguments: [ OSCArgument.blob(OSCUtility.ueControllerButtonData(key: "Gamepad_RightTrigger", controller: controllerIndex, isRepeat: isRepeat)) ])
-
+            inputType = .triggerButtonRight
         case gp.dpad.up:
-            osc.send(address, arguments: [ OSCArgument.blob(OSCUtility.ueControllerButtonData(key: "Gamepad_DPad_Up", controller: controllerIndex, isRepeat: isRepeat)) ])
+            inputType = .dpadUp
         case gp.dpad.down:
-            osc.send(address, arguments: [ OSCArgument.blob(OSCUtility.ueControllerButtonData(key: "Gamepad_DPad_Down", controller: controllerIndex, isRepeat: isRepeat)) ])
+            inputType = .dpadDown
         case gp.dpad.left:
-            osc.send(address, arguments: [ OSCArgument.blob(OSCUtility.ueControllerButtonData(key: "Gamepad_DPad_Left", controller: controllerIndex, isRepeat: isRepeat)) ])
+            inputType = .dpadLeft
         case gp.dpad.right:
-            osc.send(address, arguments: [ OSCArgument.blob(OSCUtility.ueControllerButtonData(key: "Gamepad_DPad_Right", controller: controllerIndex, isRepeat: isRepeat)) ])
-            
+            inputType = .dpadRight
         case gp.buttonOptions:
-            osc.send(address, arguments: [ OSCArgument.blob(OSCUtility.ueControllerButtonData(key: "Gamepad_Special_Left", controller: controllerIndex, isRepeat: isRepeat)) ])
+            inputType = .specialButtonLeft
         case gp.buttonMenu:
-            osc.send(address, arguments: [ OSCArgument.blob(OSCUtility.ueControllerButtonData(key: "Gamepad_Special_Right", controller: controllerIndex, isRepeat: isRepeat)) ])
-
+            inputType = .specialButtonRight
         default:
-            Log.warning("buttonValueChanged() encounted unsupported '\(button.localizedName ?? "UNK")'")
+            Log.warning("Couldn't find mapping for input button \(button.description)")
+            return
+        }
+        
+        if button.isPressed {
+            sc.sendControllerButtonPressed(inputType!, controllerIndex: UInt8(controllerIndex), isRepeat: isRepeat)
+        } else {
+            sc.sendControllerButtonReleased(inputType!, controllerIndex: UInt8(controllerIndex))
         }
-         */
     }
 }
 
diff --git a/Engine/Extras/VirtualProduction/LiveLinkVCAM/vcam/WebRTCStreamingConnection.swift b/Engine/Extras/VirtualProduction/LiveLinkVCAM/vcam/WebRTCStreamingConnection.swift
--- a/Engine/Extras/VirtualProduction/LiveLinkVCAM/vcam/WebRTCStreamingConnection.swift
+++ b/Engine/Extras/VirtualProduction/LiveLinkVCAM/vcam/WebRTCStreamingConnection.swift
@@ -200,6 +200,96 @@ class WebRTCStreamingConnection : StreamingConnection {
         client.sendData(Data(bytes))
     }
     
+    func inputTypeToIndex(_ type : StreamingConnectionControllerInputType) -> UInt8? {
+        switch type {
+        case .faceButtonBottom : return 0
+        case .faceButtonRight : return 1
+        case .faceButtonLeft : return 2
+        case .faceButtonTop : return 3
+            
+        case .shoulderButtonLeft : return 4
+        case .shoulderButtonRight : return 5
+
+        case .triggerButtonLeft : return 6
+        case .triggerButtonRight : return 7
+
+        case .specialButtonLeft : return 8
+        case .specialButtonRight : return 9
+            
+        case .thumbstickLeftButton : return 10
+        case .thumbstickRightButton : return 11
+            
+        case .dpadUp : return 12
+        case .dpadDown : return 13
+        case .dpadLeft : return 14
+        case .dpadRight : return 15
+            
+        case .thumbstickLeftX : return 1
+        case .thumbstickLeftY : return 2
+        case .thumbstickRightX : return 3
+        case .thumbstickRightY : return 4
+        }
+    }
+    
+    override func sendControllerAnalog(_ type : StreamingConnectionControllerInputType, controllerIndex : UInt8, value : Float) {
+        
+        guard let client = webRTCClient else { return }
+
+        guard let analogIndex = inputTypeToIndex(type) else {
+            Log.warning("Couldn't find an index for controller input type \(type.rawValue)")
+            return
+        }
+        
+        var bytes: [UInt8] = []
+
+        // Write message type using 1 byte
+        bytes.append(PixelStreamingToStreamerMessage.GamepadAnalog.rawValue)
+
+        bytes.append(controllerIndex)
+        bytes.append(analogIndex)
+        bytes.append(contentsOf: Double(value).toBytes())
+        
+        client.sendData(Data(bytes))
+    }
+    
+    override func sendControllerButtonPressed(_ type : StreamingConnectionControllerInputType, controllerIndex : UInt8, isRepeat : Bool) {
+
+        guard let client = webRTCClient else { return }
+
+        guard let buttonIndex = inputTypeToIndex(type) else {
+            Log.warning("Couldn't find an index for controller input type \(type.rawValue)")
+            return
+        }
+
+        var bytes: [UInt8] = []
+
+        bytes.append(PixelStreamingToStreamerMessage.GamepadButtonPressed.rawValue)
+        bytes.append(controllerIndex)
+        bytes.append(buttonIndex)
+        bytes.append(UInt8(isRepeat ? 1 : 0))
+        
+        client.sendData(Data(bytes))
+    }
+
+   override func sendControllerButtonReleased(_ type : StreamingConnectionControllerInputType, controllerIndex : UInt8) {
+
+       guard let client = webRTCClient else { return }
+
+       guard let buttonIndex = inputTypeToIndex(type) else {
+           Log.warning("Couldn't find an index for controller input type \(type.rawValue)")
+           return
+       }
+
+       var bytes: [UInt8] = []
+
+       bytes.append(PixelStreamingToStreamerMessage.GamepadButtonReleased.rawValue)
+       bytes.append(controllerIndex)
+       bytes.append(buttonIndex)
+       
+       client.sendData(Data(bytes))
+    }
+
+    
     func attachVideoTrack() {
         if let webRTC = webRTCClient, let view = self.webRTCView, let track = self.rtcVideoTrack {
             self.touchControls = TouchControls(webRTC, touchView: view)
