Live Link VCAM v1.3.0 build 6 : Fixed a bug where a client was unable to start a new connection after ending a previous connection (error was "Error Couldn't connect: Manually closed websocket").

diff --git a/Engine/Extras/VirtualProduction/LiveLinkVCAM/vcam/RemoteSessionStreamingConnection.swift b/Engine/Extras/VirtualProduction/LiveLinkVCAM/vcam/RemoteSessionStreamingConnection.swift
--- a/Engine/Extras/VirtualProduction/LiveLinkVCAM/vcam/RemoteSessionStreamingConnection.swift
+++ b/Engine/Extras/VirtualProduction/LiveLinkVCAM/vcam/RemoteSessionStreamingConnection.swift
@@ -91,6 +91,14 @@ class RemoteSessionStreamingConnection : StreamingConnection  {
         
         restartLiveLink()
     }
+    
+    deinit {
+        Log.info("Destroyed RemoteSessionStreamingConnection")
+    }
+    
+    override func shutdown() {
+        disconnect()
+    }
 
     override func connect() throws {
         self.oscConnection = nil
@@ -101,8 +109,6 @@ class RemoteSessionStreamingConnection : StreamingConnection  {
             if let conn = self.oscConnection {
                 if conn.isConnected {
                     conn.send(OSCAddressPattern.ping)
-                } else {
-                    conn.reconnect()
                 }
             }
             
diff --git a/Engine/Extras/VirtualProduction/LiveLinkVCAM/vcam/StartViewController+StreamingConnectionDelegate.swift b/Engine/Extras/VirtualProduction/LiveLinkVCAM/vcam/StartViewController+StreamingConnectionDelegate.swift
--- a/Engine/Extras/VirtualProduction/LiveLinkVCAM/vcam/StartViewController+StreamingConnectionDelegate.swift
+++ b/Engine/Extras/VirtualProduction/LiveLinkVCAM/vcam/StartViewController+StreamingConnectionDelegate.swift
@@ -24,6 +24,7 @@ extension StartViewController : StreamingConnectionDelegate {
 
         DispatchQueue.main.async {
 
+            connection.disconnect()
             self.hideConnectingAlertView() {
                 if let e = err {
                     Log.info("StreamingConnection \(connection.name) disconnected with error : \(e.localizedDescription)")
diff --git a/Engine/Extras/VirtualProduction/LiveLinkVCAM/vcam/StartViewController.swift b/Engine/Extras/VirtualProduction/LiveLinkVCAM/vcam/StartViewController.swift
--- a/Engine/Extras/VirtualProduction/LiveLinkVCAM/vcam/StartViewController.swift
+++ b/Engine/Extras/VirtualProduction/LiveLinkVCAM/vcam/StartViewController.swift
@@ -112,6 +112,7 @@ class StartViewController : BaseViewController {
         // initial & value changes for the connection type instantiates a new StreamingConnection object
         observers.append(observe(\.appSettings.connectionType, options: [.initial, .old,.new], changeHandler: { object, change in
             
+            self.streamingConnection?.shutdown()
             self.streamingConnection = nil
 
             let connectionType = self.appSettings.connectionType
@@ -149,6 +150,7 @@ class StartViewController : BaseViewController {
     
     override func viewDidAppear(_ animated: Bool) {
         super.viewDidAppear(animated)
+        self.streamingConnection?.disconnect()
     }
     
     override func viewDidDisappear(_ animated: Bool) {
diff --git a/Engine/Extras/VirtualProduction/LiveLinkVCAM/vcam/StreamingConnection.swift b/Engine/Extras/VirtualProduction/LiveLinkVCAM/vcam/StreamingConnection.swift
--- a/Engine/Extras/VirtualProduction/LiveLinkVCAM/vcam/StreamingConnection.swift
+++ b/Engine/Extras/VirtualProduction/LiveLinkVCAM/vcam/StreamingConnection.swift
@@ -9,7 +9,7 @@
 import Foundation
 import UIKit
 
-protocol StreamingConnectionDelegate {
+protocol StreamingConnectionDelegate : AnyObject {
     
     func streamingConnectionDidConnect(_ connection : StreamingConnection)
     func streamingConnection(_ connection : StreamingConnection, didDisconnectWithError err: Error?)
@@ -52,7 +52,7 @@ class StreamingConnectionStats {
 
 class StreamingConnection : NSObject {
 
-    var delegate: StreamingConnectionDelegate?
+    weak var delegate: StreamingConnectionDelegate?
     var subjectName : String!
     
     var renderView : UIView?
@@ -98,9 +98,12 @@ class StreamingConnection : NSObject {
     required init(subjectName: String) {
         self.subjectName = subjectName
     }
+    
+    func shutdown() {
+        assertionFailure("not implemented")
+    }
 
     func connect() throws {
-        assertionFailure("not implemented")
     }
 
     func reconnect() {
diff --git a/Engine/Extras/VirtualProduction/LiveLinkVCAM/vcam/VideoViewController+StreamingConnectionDelegate.swift b/Engine/Extras/VirtualProduction/LiveLinkVCAM/vcam/VideoViewController+StreamingConnectionDelegate.swift
--- a/Engine/Extras/VirtualProduction/LiveLinkVCAM/vcam/VideoViewController+StreamingConnectionDelegate.swift
+++ b/Engine/Extras/VirtualProduction/LiveLinkVCAM/vcam/VideoViewController+StreamingConnectionDelegate.swift
@@ -29,7 +29,7 @@ extension VideoViewController : StreamingConnectionDelegate {
             if self.dismissOnDisconnect {
                 self.forceDisconnectAndDismiss()
             } else {
-                self.showReconnecting(true, animated: true)
+                self.reconnect()
             }
         }
 
diff --git a/Engine/Extras/VirtualProduction/LiveLinkVCAM/vcam/VideoViewController.swift b/Engine/Extras/VirtualProduction/LiveLinkVCAM/vcam/VideoViewController.swift
--- a/Engine/Extras/VirtualProduction/LiveLinkVCAM/vcam/VideoViewController.swift
+++ b/Engine/Extras/VirtualProduction/LiveLinkVCAM/vcam/VideoViewController.swift
@@ -140,48 +140,6 @@ class VideoViewController : BaseViewController {
         
         self.streamingConnection?.renderView = self.renderView
         
-        /*
-        
-        arView.session.run(config, options: [ .resetTracking ] )
-        arView.session.delegate = self
-        arView.delegate = self
-        
-        if self.demoMode {
-            arView.automaticallyUpdatesLighting = true
-            arView.debugOptions = [ .showWorldOrigin ]
-
-            let floor = SCNNode(geometry: SCNPlane(width: 14, height: 14))
-            floor.rotation = SCNVector4(1, 0, 0, GLKMathDegreesToRadians(-90))
-            floor.position = SCNVector3(0, -1, 0)
-            floor.geometry?.firstMaterial?.diffuse.contents = UIImage(named: "checkerboard")
-            floor.geometry?.firstMaterial?.diffuse.contentsTransform = SCNMatrix4MakeScale(50, 50, 0)
-            floor.geometry?.firstMaterial?.diffuse.wrapS = .repeat
-            floor.geometry?.firstMaterial?.diffuse.wrapT = .repeat
-            arView.scene.rootNode.addChildNode(floor)
-
-            for x in 0...1 {
-                for z in 0...1 {
-                    let box = SCNNode(geometry: SCNBox(width: 1, height: 1, length: 1, chamferRadius: 0))
-                    box.position = SCNVector3(-5.0 + Double(x) * 10.0, -0.5, -5.0 + Double(z) * 10.0)
-                    box.rotation = SCNVector4(0,1,0, Float.random(in: 0...3.14))
-                    for i in 0...5 {
-                        let mat = SCNMaterial()
-                        mat.diffuse.contents = UIImage(named: "UnrealLogo")
-                        if i == 0 {
-                            box.geometry?.firstMaterial = mat
-                        } else {
-                            box.geometry?.materials.append(mat)
-                        }
-                    }
-                    arView.scene.rootNode.addChildNode(box)
-                }
-            }
-
-        } else {
-            arView.scene.background.contents = UIColor.black
-        }
-        
-         */
         
         let coachingOverlayView = ARCoachingOverlayView()
         
@@ -241,6 +199,8 @@ class VideoViewController : BaseViewController {
     }
     
     func reconnect() {
+        
+        showReconnecting(true, animated: true)
         self.streamingConnection?.reconnect()
     }
     
diff --git a/Engine/Extras/VirtualProduction/LiveLinkVCAM/vcam/WebRTCStreamingConnection.swift b/Engine/Extras/VirtualProduction/LiveLinkVCAM/vcam/WebRTCStreamingConnection.swift
--- a/Engine/Extras/VirtualProduction/LiveLinkVCAM/vcam/WebRTCStreamingConnection.swift
+++ b/Engine/Extras/VirtualProduction/LiveLinkVCAM/vcam/WebRTCStreamingConnection.swift
@@ -113,40 +113,51 @@ class WebRTCStreamingConnection : StreamingConnection {
     }
     
     deinit {
-        _statsTimer?.invalidate()
+        
+        Log.info("Destroyed WebRTCStreamingConnection")
+
         webRTCClient = nil
     }
 
-    override func connect() throws {
-        
+    override func shutdown() {
+        _statsTimer?.invalidate()
         disconnect()
+    }
+    
+    override func connect() throws {
         
-        // We will use 3rd party library for websockets.
-        let webSocketProvider: WebSocketProvider
-        
-        if #available(iOS 13.0, *) {
-            webSocketProvider = NativeWebSocketProvider(url: self._url!, timeout: 2.0)
-        } else {
-            webSocketProvider = StarscreamWebSocket(url: self._url!)
+        if signalClient == nil {
+            
+            // We will use 3rd party library for websockets.
+            let webSocketProvider: WebSocketProvider
+            
+            if #available(iOS 13.0, *) {
+                webSocketProvider = NativeWebSocketProvider(url: self._url!, timeout: 2.0)
+            } else {
+                webSocketProvider = StarscreamWebSocket(url: self._url!)
+            }
+            
+            self.signalClient = SignalingClient(webSocket: webSocketProvider)
+            self.signalClient?.delegate = self
+            self.signalClient?.connect()
+
         }
-        
-        self.signalClient = SignalingClient(webSocket: webSocketProvider)
-        self.signalClient?.delegate = self
-        self.signalClient?.connect()
     }
     
     override func reconnect() {
-        disconnect()
-        do {
-            try connect()
-        } catch {
-            Log.error(error.localizedDescription)
-        }
+        //disconnect()
+        //do {
+        //    try connect()
+        //} catch {
+        //    Log.error(error.localizedDescription)
+        //}
     }
     
     override func disconnect() {
-        self.signalClient?.close()
-        self.signalClient = nil
+        signalClient?.delegate = nil
+        signalClient?.close()
+        signalClient = nil
+
         signalingConnected = false
         hasLocalSdp = false
         hasRemoteSdp = false
@@ -420,6 +431,8 @@ extension WebRTCStreamingConnection: WebRTCClientDelegate {
             
         case .disconnected:
             self.delegate?.streamingConnection(self, didDisconnectWithError: nil)
+        case .failed:
+            self.delegate?.streamingConnection(self, didDisconnectWithError: NSError(domain: "", code: 1, userInfo: [NSLocalizedDescriptionKey : "Failed to connect." ] ))
         default:
             break
         }
