[EpicStageApp] Reconnect option

diff --git a/Engine/Extras/VirtualProduction/EpicStageApp/IOS/Flutter/Debug.xcconfig b/Engine/Extras/VirtualProduction/EpicStageApp/IOS/Flutter/Debug.xcconfig
--- a/Engine/Extras/VirtualProduction/EpicStageApp/IOS/Flutter/Debug.xcconfig
+++ b/Engine/Extras/VirtualProduction/EpicStageApp/IOS/Flutter/Debug.xcconfig
@@ -1 +1,2 @@
+#include? "Pods/Target Support Files/Pods-Runner/Pods-Runner.debug.xcconfig"
 #include "Generated.xcconfig"
diff --git a/Engine/Extras/VirtualProduction/EpicStageApp/IOS/Flutter/Release.xcconfig b/Engine/Extras/VirtualProduction/EpicStageApp/IOS/Flutter/Release.xcconfig
--- a/Engine/Extras/VirtualProduction/EpicStageApp/IOS/Flutter/Release.xcconfig
+++ b/Engine/Extras/VirtualProduction/EpicStageApp/IOS/Flutter/Release.xcconfig
@@ -1 +1,2 @@
+#include? "Pods/Target Support Files/Pods-Runner/Pods-Runner.release.xcconfig"
 #include "Generated.xcconfig"
