Live Link VCAM iOS (v1.3.0 build 7) : Defaults connection type to PixelStreaming/WebRTC.

diff --git a/Engine/Extras/VirtualProduction/LiveLinkVCAM/vcam/AppSettings.swift b/Engine/Extras/VirtualProduction/LiveLinkVCAM/vcam/AppSettings.swift
--- a/Engine/Extras/VirtualProduction/LiveLinkVCAM/vcam/AppSettings.swift
+++ b/Engine/Extras/VirtualProduction/LiveLinkVCAM/vcam/AppSettings.swift
@@ -53,7 +53,7 @@ class AppSettings : NSObject {
 
             Keys.showStreamingStats : false,
             
-            Keys.connectionType : "RemoteSession"
+            Keys.connectionType : StreamingConnectionType.webRTC.rawValue
 
         ]
     }
@@ -134,11 +134,19 @@ class AppSettings : NSObject {
         }
     }
     
-    @objc dynamic var connectionType : String = UserDefaults.standard.string(forKey: Keys.connectionType) ?? "RemoteSession" {
+    @objc dynamic var connectionType : String = UserDefaults.standard.string(forKey: Keys.connectionType) ?? StreamingConnectionType.webRTC.rawValue {
         didSet {
             UserDefaults.standard.set(connectionType, forKey: Keys.connectionType)
         }
     }
     
+    func connectionTypeEnum() -> StreamingConnectionType {
+        return StreamingConnectionType(rawValue: self.connectionType) ?? .webRTC
+    }
+
+    func setConnectionTypeEnum(_ type : StreamingConnectionType)  {
+        self.connectionType = type.rawValue
+    }
+    
 
 }
diff --git a/Engine/Extras/VirtualProduction/LiveLinkVCAM/vcam/RemoteSessionStreamingConnection.swift b/Engine/Extras/VirtualProduction/LiveLinkVCAM/vcam/RemoteSessionStreamingConnection.swift
--- a/Engine/Extras/VirtualProduction/LiveLinkVCAM/vcam/RemoteSessionStreamingConnection.swift
+++ b/Engine/Extras/VirtualProduction/LiveLinkVCAM/vcam/RemoteSessionStreamingConnection.swift
@@ -53,7 +53,7 @@ class RemoteSessionStreamingConnection : StreamingConnection  {
     
     override var name : String {
         get {
-            "RemoteSession"
+            StreamingConnectionType.remoteSession.rawValue
         }
     }
     
diff --git a/Engine/Extras/VirtualProduction/LiveLinkVCAM/vcam/SettingsViewController.swift b/Engine/Extras/VirtualProduction/LiveLinkVCAM/vcam/SettingsViewController.swift
--- a/Engine/Extras/VirtualProduction/LiveLinkVCAM/vcam/SettingsViewController.swift
+++ b/Engine/Extras/VirtualProduction/LiveLinkVCAM/vcam/SettingsViewController.swift
@@ -40,13 +40,11 @@ class SettingsViewController : UITableViewController {
             cell.detailTextLabel?.text = Timecode.sourceToString(appSettings.timecodeSourceEnum())
             
         case "connectionType":
-            switch appSettings.connectionType {
-            case "RemoteSession":
+            switch appSettings.connectionTypeEnum() {
+            case .remoteSession:
                 cell.detailTextLabel?.text = "Remote Session"
-            case "WebRTC":
+            case .webRTC:
                 cell.detailTextLabel?.text = "Pixel Streaming"
-            default:
-                cell.detailTextLabel?.text = "Unknown"
             }
 
         default:
diff --git a/Engine/Extras/VirtualProduction/LiveLinkVCAM/vcam/StreamingConnection.swift b/Engine/Extras/VirtualProduction/LiveLinkVCAM/vcam/StreamingConnection.swift
--- a/Engine/Extras/VirtualProduction/LiveLinkVCAM/vcam/StreamingConnection.swift
+++ b/Engine/Extras/VirtualProduction/LiveLinkVCAM/vcam/StreamingConnection.swift
@@ -15,6 +15,11 @@ protocol StreamingConnectionDelegate : AnyObject {
     func streamingConnection(_ connection : StreamingConnection, didDisconnectWithError err: Error?)
 }
 
+enum StreamingConnectionType : String {
+    case remoteSession = "RemoteSession"
+    case webRTC = "WebRTC"
+}
+
 enum StreamingConnectionTouchType : String {
     case began = "Began"
     case moved = "Moved"
diff --git a/Engine/Extras/VirtualProduction/LiveLinkVCAM/vcam/WebRTCStreamingConnection.swift b/Engine/Extras/VirtualProduction/LiveLinkVCAM/vcam/WebRTCStreamingConnection.swift
--- a/Engine/Extras/VirtualProduction/LiveLinkVCAM/vcam/WebRTCStreamingConnection.swift
+++ b/Engine/Extras/VirtualProduction/LiveLinkVCAM/vcam/WebRTCStreamingConnection.swift
@@ -30,7 +30,7 @@ class WebRTCStreamingConnection : StreamingConnection {
 
     override var name : String {
         get {
-            "WebRTC"
+            StreamingConnectionType.webRTC.rawValue
         }
     }
     override var destination : String {
