[EpicStageApp] Web interface tab

diff --git a/Engine/Extras/VirtualProduction/EpicStageApp/android/app/src/main/AndroidManifest.xml b/Engine/Extras/VirtualProduction/EpicStageApp/android/app/src/main/AndroidManifest.xml
--- a/Engine/Extras/VirtualProduction/EpicStageApp/android/app/src/main/AndroidManifest.xml
+++ b/Engine/Extras/VirtualProduction/EpicStageApp/android/app/src/main/AndroidManifest.xml
@@ -3,7 +3,8 @@
    <application
         android:label="epic_stage"
         android:name="${applicationName}"
-        android:icon="@mipmap/ic_launcher">
+        android:icon="@mipmap/ic_launcher"
+	    android:usesCleartextTraffic="true">
         <activity
             android:name=".MainActivity"
             android:exported="true"
diff --git a/Engine/Extras/VirtualProduction/EpicStageApp/android/app/src/main/res/values/styles.xml b/Engine/Extras/VirtualProduction/EpicStageApp/android/app/src/main/res/values/styles.xml
--- a/Engine/Extras/VirtualProduction/EpicStageApp/android/app/src/main/res/values/styles.xml
+++ b/Engine/Extras/VirtualProduction/EpicStageApp/android/app/src/main/res/values/styles.xml
@@ -12,7 +12,7 @@
          running.
 
          This Theme is only used starting with V2 of Flutter's Android embedding. -->
-    <style name="NormalTheme" parent="@android:style/Theme.Light.NoTitleBar">
+	<style name="NormalTheme" parent="Theme.MaterialComponents.DayNight.NoActionBar">
         <item name="android:windowBackground">?android:colorBackground</item>
     </style>
 </resources>
