TextureShare: added missing SDK files to allow Samples to compile

diff --git a/Engine/Extras/VirtualProduction/TextureShare/TextureShareSDK/Source/Containers/TextureShareCoreContainers.h b/Engine/Extras/VirtualProduction/TextureShare/TextureShareSDK/Source/Containers/TextureShareCoreContainers.h
new file mode 100644
--- /dev/null
+++ b/Engine/Extras/VirtualProduction/TextureShare/TextureShareSDK/Source/Containers/TextureShareCoreContainers.h
@@ -0,0 +1,50 @@
+// Copyright Epic Games, Inc. All Rights Reserved.
+
+#pragma once
+#include "Containers/TextureShareCoreContainers_Data.h"
+#include "Containers/TextureShareCoreContainers_ProxyData.h"
+#include "Containers/TextureShareCoreContainers_SyncSettings.h"
+
+/**
+ * TextureShareCore object data container
+ * The structure is serializable (binary compatible) and reflected on the SDK side (UE types will be replaced with simplified copies from the SDK).
+ */
+struct FTextureShareCoreObjectData
+	: public ITextureShareSerialize
+{
+	// Object information
+	FTextureShareCoreObjectDesc Desc;
+
+	// Setup frame data on game thread before rendering call
+	FTextureShareCoreData Data;
+
+public:
+	virtual ~FTextureShareCoreObjectData() = default;
+
+	virtual ITextureShareSerializeStream& Serialize(ITextureShareSerializeStream& Stream) override
+	{
+		return Stream << Desc << Data;
+	}
+};
+
+/**
+ * TextureShareCore object proxy data container
+ * The structure is serializable (binary compatible) and reflected on the SDK side (UE types will be replaced with simplified copies from the SDK).
+ */
+struct FTextureShareCoreObjectProxyData
+	: public ITextureShareSerialize
+{
+	// Object ID
+	FTextureShareCoreObjectDesc Desc;
+
+	// Setup frame data on game thread before rendering call
+	FTextureShareCoreProxyData ProxyData;
+
+public:
+	virtual ~FTextureShareCoreObjectProxyData() = default;
+
+	virtual ITextureShareSerializeStream& Serialize(ITextureShareSerializeStream & Stream) override
+	{
+		return Stream << Desc << ProxyData;
+	}
+};
diff --git a/Engine/Extras/VirtualProduction/TextureShare/TextureShareSDK/Source/Containers/TextureShareCoreContainers_CustomData.h b/Engine/Extras/VirtualProduction/TextureShare/TextureShareSDK/Source/Containers/TextureShareCoreContainers_CustomData.h
new file mode 100644
--- /dev/null
+++ b/Engine/Extras/VirtualProduction/TextureShare/TextureShareSDK/Source/Containers/TextureShareCoreContainers_CustomData.h
@@ -0,0 +1,57 @@
+// Copyright Epic Games, Inc. All Rights Reserved.
+
+#pragma once
+#include "Serialize/TextureShareCoreSerialize.h"
+
+/**
+ * Custom data pair: Key + Value
+ * The structure is serializable (binary compatible) and reflected on the SDK side (UE types will be replaced with simplified copies from the SDK).
+ */
+struct FTextureShareCoreCustomData
+	: public ITextureShareSerialize
+{
+	// Custom data key
+	FString Key;
+
+	// Custom data value
+	FString Value;
+
+public:
+	virtual ~FTextureShareCoreCustomData() = default;
+
+	virtual ITextureShareSerializeStream& Serialize(ITextureShareSerializeStream & Stream) override
+	{
+		return Stream << Key << Value;
+	}
+
+public:
+	FTextureShareCoreCustomData() = default;
+
+	FTextureShareCoreCustomData(const wchar_t* InKey, const wchar_t* InValue)
+		: Key(InKey), Value(InValue)
+	{}
+
+	FTextureShareCoreCustomData(const FString& InKey, const FString& InValue)
+		: Key(InKey), Value(InValue)
+	{}
+
+	bool EqualsFunc(const wchar_t* InKey) const
+	{
+		return Key == InKey;
+	}
+
+	bool EqualsFunc(const FString& InKey) const
+	{
+		return Key == InKey;
+	}
+
+	bool EqualsFunc(const FTextureShareCoreCustomData& InParameter) const
+	{
+		return Key == InParameter.Key;
+	}
+
+	bool operator==(const FTextureShareCoreCustomData& InParameter) const
+	{
+		return Key == InParameter.Key;
+	}
+};
diff --git a/Engine/Extras/VirtualProduction/TextureShare/TextureShareSDK/Source/Containers/TextureShareCoreContainers_Data.h b/Engine/Extras/VirtualProduction/TextureShare/TextureShareSDK/Source/Containers/TextureShareCoreContainers_Data.h
new file mode 100644
--- /dev/null
+++ b/Engine/Extras/VirtualProduction/TextureShare/TextureShareSDK/Source/Containers/TextureShareCoreContainers_Data.h
@@ -0,0 +1,52 @@
+// Copyright Epic Games, Inc. All Rights Reserved.
+
+#pragma once
+#include "Containers/TextureShareCoreContainers_FrameMarker.h"
+#include "Containers/TextureShareCoreContainers_ViewDesc.h"
+#include "Containers/TextureShareCoreContainers_ResourceRequest.h"
+#include "Containers/TextureShareCoreContainers_ManualProjection.h"
+#include "Containers/TextureShareCoreContainers_CustomData.h"
+
+/**
+ * Game data of the TextureShareCore object
+ * The structure is serializable (binary compatible) and reflected on the SDK side (UE types will be replaced with simplified copies from the SDK).
+ */
+struct FTextureShareCoreData
+	: public ITextureShareSerialize
+{
+	// Current frame marker
+	FTextureShareCoreFrameMarker FrameMarker;
+
+	// Supported views of this object
+	TArraySerializable<FTextureShareCoreViewDesc> SupportedViews;
+
+	// Resource requests from this object to remote processes
+	TArraySerializable<FTextureShareCoreResourceRequest> ResourceRequests;
+
+	// Manual projection requests from this object to remote processes
+	TArraySerializable<FTextureShareCoreManualProjection> ManualProjections;
+
+	// Sources of used manual projections in local process
+	TArraySerializable<FTextureShareCoreManualProjectionSource> ManualProjectionsSources;
+
+	// Custom abstract data of this object
+	TArraySerializable<FTextureShareCoreCustomData> CustomData;
+
+public:
+	virtual ~FTextureShareCoreData() = default;
+
+	virtual ITextureShareSerializeStream& Serialize(ITextureShareSerializeStream & Stream) override
+	{
+		return Stream << FrameMarker << SupportedViews << ResourceRequests << ManualProjections << ManualProjectionsSources << CustomData;
+	}
+
+public:
+	void ResetData()
+	{
+		SupportedViews.Empty();
+		ResourceRequests.Empty();
+		ManualProjections.Empty();
+		ManualProjectionsSources.Empty();
+		CustomData.Empty();
+	}
+};
diff --git a/Engine/Extras/VirtualProduction/TextureShare/TextureShareSDK/Source/Containers/TextureShareCoreContainers_DeviceD3D11.h b/Engine/Extras/VirtualProduction/TextureShare/TextureShareSDK/Source/Containers/TextureShareCoreContainers_DeviceD3D11.h
new file mode 100644
--- /dev/null
+++ b/Engine/Extras/VirtualProduction/TextureShare/TextureShareSDK/Source/Containers/TextureShareCoreContainers_DeviceD3D11.h
@@ -0,0 +1,6 @@
+// Copyright Epic Games, Inc. All Rights Reserved.
+
+#pragma once
+
+struct ID3D11Device;
+struct ID3D11Texture2D;
diff --git a/Engine/Extras/VirtualProduction/TextureShare/TextureShareSDK/Source/Containers/TextureShareCoreContainers_DeviceD3D12.h b/Engine/Extras/VirtualProduction/TextureShare/TextureShareSDK/Source/Containers/TextureShareCoreContainers_DeviceD3D12.h
new file mode 100644
--- /dev/null
+++ b/Engine/Extras/VirtualProduction/TextureShare/TextureShareSDK/Source/Containers/TextureShareCoreContainers_DeviceD3D12.h
@@ -0,0 +1,9 @@
+// Copyright Epic Games, Inc. All Rights Reserved.
+
+#pragma once
+
+struct ID3D12Device;
+struct ID3D12Resource;
+
+struct ID3D12GraphicsCommandList;
+struct ID3D12DescriptorHeap;
diff --git a/Engine/Extras/VirtualProduction/TextureShare/TextureShareSDK/Source/Containers/TextureShareCoreContainers_DeviceVulkan.h b/Engine/Extras/VirtualProduction/TextureShare/TextureShareSDK/Source/Containers/TextureShareCoreContainers_DeviceVulkan.h
new file mode 100644
--- /dev/null
+++ b/Engine/Extras/VirtualProduction/TextureShare/TextureShareSDK/Source/Containers/TextureShareCoreContainers_DeviceVulkan.h
@@ -0,0 +1,113 @@
+// Copyright Epic Games, Inc. All Rights Reserved.
+
+#pragma once
+#include "Serialize/TextureShareCoreSerialize.h"
+
+typedef struct VkInstance_T*       VkInstance;
+typedef struct VkPhysicalDevice_T* VkPhysicalDevice;
+typedef struct VkDevice_T*         VkDevice;
+typedef struct VkDeviceMemory_T*   VkDeviceMemory;
+typedef struct VkImage_T*          VkImage;
+
+/**
+ * Vulkan device context
+ * The structure is serializable (binary compatible) and reflected on the SDK side (UE types will be replaced with simplified copies from the SDK).
+ */
+struct FTextureShareDeviceVulkanContext
+	: public ITextureShareSerialize
+{
+	FTextureShareDeviceVulkanContext() = default;
+	FTextureShareDeviceVulkanContext(VkInstance InInstance, VkPhysicalDevice InPhysicalDevice, VkDevice InDevice)
+	{
+		Instance = (Windows::HANDLE)(InInstance);
+		PhysicalDevice = (Windows::HANDLE)(InPhysicalDevice);
+		Device = (Windows::HANDLE)(InDevice);
+	}
+
+	bool IsValid() const
+	{
+		return Instance != nullptr && Device!= nullptr && PhysicalDevice!=nullptr;
+	}
+
+	VkInstance GetInstance() const
+	{
+		return (VkInstance)(Instance);
+	}
+
+	VkPhysicalDevice GetPhysicalDevice() const
+	{
+		return (VkPhysicalDevice)(PhysicalDevice);
+	}
+
+	VkDevice GetDevice() const
+	{
+		return (VkDevice)(Device);
+	}
+
+public:
+	virtual ~FTextureShareDeviceVulkanContext() = default;
+
+	virtual ITextureShareSerializeStream& Serialize(ITextureShareSerializeStream & Stream) override
+	{
+		return Stream << Instance << PhysicalDevice << Device;
+	}
+
+private:
+	Windows::HANDLE Instance = nullptr;
+	Windows::HANDLE PhysicalDevice = nullptr;
+	Windows::HANDLE Device = nullptr;
+};
+
+/**
+ * Vulkan resource context
+ * The structure is serializable (binary compatible) and reflected on the SDK side (UE types will be replaced with simplified copies from the SDK).
+ */
+struct FTextureShareDeviceVulkanResource
+	: public ITextureShareSerialize
+{
+	FTextureShareDeviceVulkanResource() = default;
+
+	FTextureShareDeviceVulkanResource(VkDeviceMemory InAllocationHandle, const uint64 InAllocationOffset = 0)
+	{
+		AllocationHandle = (Windows::HANDLE)(InAllocationHandle);
+		AllocationOffset = InAllocationOffset;
+	}
+
+	void operator==(const FTextureShareDeviceVulkanResource& In)
+	{
+		AllocationHandle = In.AllocationHandle;
+		AllocationOffset = In.AllocationOffset;
+	}
+
+	bool IsValid() const
+	{
+		return AllocationHandle != nullptr;
+	}
+
+	void* GetNativeResourcePtr() const
+	{
+		return AllocationHandle;
+	}
+
+	VkDeviceMemory GetAllocationHandle() const
+	{
+		return (VkDeviceMemory)(AllocationHandle);
+	}
+
+	uint64 GetAllocationOffset() const
+	{
+		return AllocationOffset;
+	}
+
+public:
+	virtual ~FTextureShareDeviceVulkanResource() = default;
+
+	virtual ITextureShareSerializeStream& Serialize(ITextureShareSerializeStream & Stream) override
+	{
+		return Stream << AllocationHandle << AllocationOffset;
+	}
+
+private:
+	Windows::HANDLE AllocationHandle = nullptr;
+	uint64          AllocationOffset = 0;
+};
diff --git a/Engine/Extras/VirtualProduction/TextureShare/TextureShareSDK/Source/Containers/TextureShareCoreContainers_FrameMarker.h b/Engine/Extras/VirtualProduction/TextureShare/TextureShareSDK/Source/Containers/TextureShareCoreContainers_FrameMarker.h
new file mode 100644
--- /dev/null
+++ b/Engine/Extras/VirtualProduction/TextureShare/TextureShareSDK/Source/Containers/TextureShareCoreContainers_FrameMarker.h
@@ -0,0 +1,75 @@
+// Copyright Epic Games, Inc. All Rights Reserved.
+
+#pragma once
+#include "Containers/TextureShareCoreContainers_ObjectDesc.h"
+
+/**
+ * Frame marker. Sync purpose
+ * The structure is serializable (binary compatible) and reflected on the SDK side (UE types will be replaced with simplified copies from the SDK).
+ */
+struct FTextureShareCoreFrameMarker
+	: public ITextureShareSerialize
+{
+	uint32 CustomFrameIndex1 = 0;
+	uint64 CustomFrameIndex2 = 0;
+
+public:
+	virtual ~FTextureShareCoreFrameMarker() = default;
+
+	virtual ITextureShareSerializeStream& Serialize(ITextureShareSerializeStream & Stream) override
+	{
+		return Stream << CustomFrameIndex1 << CustomFrameIndex2;
+	}
+
+public:
+	void NextFrame()
+	{
+		// Simple frame marker incremental update
+		CustomFrameIndex1++;
+		CustomFrameIndex2++;
+	}
+};
+
+/**
+ * Frame marker paired with Object descriptor. Sync purpose
+ * The structure is serializable (binary compatible) and reflected on the SDK side (UE types will be replaced with simplified copies from the SDK).
+ */
+struct FTextureShareCoreObjectFrameMarker
+	: public ITextureShareSerialize
+{
+	// Projection data source process Object ID
+	FTextureShareCoreObjectDesc ObjectDesc;
+
+	// Projection data source process frame marker
+	FTextureShareCoreFrameMarker FrameMarker;
+
+public:
+	virtual ~FTextureShareCoreObjectFrameMarker() = default;
+
+	virtual ITextureShareSerializeStream& Serialize(ITextureShareSerializeStream & Stream) override
+	{
+		return Stream << ObjectDesc << FrameMarker;
+	}
+
+public:
+	FTextureShareCoreObjectFrameMarker() = default;
+
+	FTextureShareCoreObjectFrameMarker(const FTextureShareCoreObjectDesc& InObjectDesc, const FTextureShareCoreFrameMarker& InFrameMarker)
+		: ObjectDesc(InObjectDesc), FrameMarker(InFrameMarker)
+	{ }
+
+	bool EqualsFunc(const FTextureShareCoreObjectDesc& InObjectDesc) const
+	{
+		return InObjectDesc.EqualsFunc(ObjectDesc);
+	}
+
+	bool EqualsFunc(const FTextureShareCoreObjectFrameMarker& InObjectFrameMarker) const
+	{
+		return InObjectFrameMarker.ObjectDesc.EqualsFunc(ObjectDesc);
+	}
+
+	bool operator==(const FTextureShareCoreObjectFrameMarker& InObjectFrameMarker) const
+	{
+		return InObjectFrameMarker.ObjectDesc == ObjectDesc;
+	}
+};
diff --git a/Engine/Extras/VirtualProduction/TextureShare/TextureShareSDK/Source/Containers/TextureShareCoreContainers_ManualProjection.h b/Engine/Extras/VirtualProduction/TextureShare/TextureShareSDK/Source/Containers/TextureShareCoreContainers_ManualProjection.h
new file mode 100644
--- /dev/null
+++ b/Engine/Extras/VirtualProduction/TextureShare/TextureShareSDK/Source/Containers/TextureShareCoreContainers_ManualProjection.h
@@ -0,0 +1,109 @@
+// Copyright Epic Games, Inc. All Rights Reserved.
+
+#pragma once
+#include "Containers/TextureShareCoreContainers_ViewDesc.h"
+
+/**
+ * Manual projection settings for specified view
+ * The structure is serializable (binary compatible) and reflected on the SDK side (UE types will be replaced with simplified copies from the SDK).
+ */
+struct FTextureShareCoreManualProjection
+	: public ITextureShareSerialize
+{
+	// Apply manual projection for this view
+	FTextureShareCoreViewDesc ViewDesc;
+
+	ETextureShareCoreSceneViewManualProjectionType ProjectionType = ETextureShareCoreSceneViewManualProjectionType::FrustumAngles;
+
+	ETextureShareViewRotationDataType   ViewRotationType = ETextureShareViewRotationDataType::Original;
+	ETextureShareViewLocationDataType   ViewLocationType = ETextureShareViewLocationDataType::Original;
+
+	// Projection matrix
+	FMatrix ProjectionMatrix = FMatrix::Identity;
+
+	// View rotation
+	FRotator ViewRotation = FRotator::ZeroRotator;
+
+	// View location
+	FVector ViewLocation = FVector::ZeroVector;
+
+	// Frustum angles
+	struct FFrustumAngles
+	{
+		float Left = -30.f;
+		float Right = 30.f;
+		float Top = 30.f;
+		float Bottom = -30.f;
+	} FrustumAngles;
+
+public:
+	virtual ~FTextureShareCoreManualProjection() = default;
+
+	virtual ITextureShareSerializeStream& Serialize(ITextureShareSerializeStream & Stream) override
+	{
+		return Stream << ViewDesc << ProjectionType << ViewRotationType << ViewLocationType
+			<< ProjectionMatrix << ViewRotation << ViewLocation
+			<< FrustumAngles.Left << FrustumAngles.Right << FrustumAngles.Top << FrustumAngles.Bottom;
+	}
+
+public:
+	bool EqualsFunc(const FString& InViewId) const
+	{
+		return ViewDesc.EqualsFunc(InViewId);
+	}
+
+	bool EqualsFunc(const FTextureShareCoreViewDesc& InViewDesc) const
+	{
+		return ViewDesc.EqualsFunc(InViewDesc);
+	}
+
+	bool EqualsFunc(const FTextureShareCoreManualProjection& InManualProjection) const
+	{
+		return ViewDesc.EqualsFunc(InManualProjection.ViewDesc);
+	}
+};
+
+/**
+ * Manual projection source info
+ * When using a manual projection from a remote process, store the description of the remote object with a frame marker
+ * The structure is serializable (binary compatible) and reflected on the SDK side (UE types will be replaced with simplified copies from the SDK).
+ */
+struct FTextureShareCoreManualProjectionSource
+	: public ITextureShareSerialize
+{
+	// View desc type of this resource (support stereo)
+	FTextureShareCoreViewDesc ViewDesc;
+
+	// remote process object and framemarker info
+	FTextureShareCoreObjectFrameMarker ObjectFrameMarker;
+
+public:
+	virtual ~FTextureShareCoreManualProjectionSource() = default;
+
+	virtual ITextureShareSerializeStream& Serialize(ITextureShareSerializeStream & Stream) override
+	{
+		return Stream << ViewDesc << ObjectFrameMarker;
+	}
+
+public:
+	FTextureShareCoreManualProjectionSource() = default;
+
+	FTextureShareCoreManualProjectionSource(const FTextureShareCoreViewDesc& InViewDesc, const FTextureShareCoreObjectFrameMarker& InObjectFrameMarker)
+		: ViewDesc(InViewDesc), ObjectFrameMarker(InObjectFrameMarker)
+	{ }
+
+	bool EqualsFunc(const FTextureShareCoreViewDesc& InViewDesc) const
+	{
+		return ViewDesc.EqualsFunc(InViewDesc);
+	}
+
+	bool EqualsFunc(const FTextureShareCoreManualProjectionSource& InManualProjectionSource) const
+	{
+		return ViewDesc.EqualsFunc(InManualProjectionSource.ViewDesc);
+	}
+
+	bool operator==(const FTextureShareCoreManualProjectionSource& InManualProjectionSource) const
+	{
+		return ViewDesc == InManualProjectionSource.ViewDesc;
+	}
+};
diff --git a/Engine/Extras/VirtualProduction/TextureShare/TextureShareSDK/Source/Containers/TextureShareCoreContainers_ObjectDesc.h b/Engine/Extras/VirtualProduction/TextureShare/TextureShareSDK/Source/Containers/TextureShareCoreContainers_ObjectDesc.h
new file mode 100644
--- /dev/null
+++ b/Engine/Extras/VirtualProduction/TextureShare/TextureShareSDK/Source/Containers/TextureShareCoreContainers_ObjectDesc.h
@@ -0,0 +1,142 @@
+// Copyright Epic Games, Inc. All Rights Reserved.
+
+#pragma once
+#include "Serialize/TextureShareCoreSerialize.h"
+
+/**
+ * Process descriptor
+ * The structure is serializable (binary compatible) and reflected on the SDK side (UE types will be replaced with simplified copies from the SDK).
+ */
+struct FTextureShareCoreObjectProcessDesc
+	: public ITextureShareSerialize
+{
+	// Process name (user-defined)
+	FString ProcessId;
+
+	// Unique process guid
+	FGuid ProcessGuid;
+
+	// The local process type
+	ETextureShareProcessType ProcessType = ETextureShareProcessType::Undefined;
+
+	// Local process device type
+	ETextureShareDeviceType DeviceType = ETextureShareDeviceType::Undefined;
+
+public:
+	virtual ~FTextureShareCoreObjectProcessDesc() = default;
+
+	virtual ITextureShareSerializeStream& Serialize(ITextureShareSerializeStream & Stream) override
+	{
+		return Stream << ProcessId << ProcessGuid << ProcessType << DeviceType;
+	}
+};
+
+/**
+ * Object sync logic state
+ * The structure is serializable (binary compatible) and reflected on the SDK side (UE types will be replaced with simplified copies from the SDK).
+ */
+struct FTextureShareCoreObjectSyncState
+	: public ITextureShareSerialize
+{
+	ETextureShareSyncStep  Step;
+	ETextureShareSyncState State;
+
+	ETextureShareSyncStep NextStep;
+	ETextureShareSyncStep PrevStep;
+
+public:
+	virtual ~FTextureShareCoreObjectSyncState() = default;
+	virtual ITextureShareSerializeStream& Serialize(ITextureShareSerializeStream& Stream) override
+	{
+		return Stream << Step << State << NextStep << PrevStep;
+	}
+};
+
+/**
+ * Object sync info
+ * The structure is serializable (binary compatible) and reflected on the SDK side (UE types will be replaced with simplified copies from the SDK).
+ */
+struct FTextureShareCoreObjectSync
+	: public ITextureShareSerialize
+{
+	// Object sync logic state
+	FTextureShareCoreObjectSyncState SyncState;
+
+	// Bit-storage for sync steps used by this process
+	// 64 bits. Max steps=64
+	uint64 SyncStepSettings;
+
+	// Last access time from this object
+	uint64 LastAccessTime = 0;
+
+public:
+	virtual ~FTextureShareCoreObjectSync() = default;
+	virtual ITextureShareSerializeStream& Serialize(ITextureShareSerializeStream& Stream) override
+	{
+		return Stream << SyncState << SyncStepSettings << LastAccessTime;
+	}
+
+public:
+	bool IsStepEnabled(const ETextureShareSyncStep InStep) const
+	{
+		if (InStep == ETextureShareSyncStep::InterprocessConnection)
+		{
+			return true;
+		}
+
+		const int8 BitIndex = (int8)InStep;
+
+		if (BitIndex >= 0 && BitIndex < 64)
+		{
+			return (SyncStepSettings & (1ULL << BitIndex)) != 0;
+		}
+
+		return false;
+	}
+};
+
+/**
+ * Object unique description
+ * The structure is serializable (binary compatible) and reflected on the SDK side (UE types will be replaced with simplified copies from the SDK).
+ */
+struct FTextureShareCoreObjectDesc
+	: public ITextureShareSerialize
+{
+	// Texture share name (TS objects with the same ShareName are trying to connect)
+	FString ShareName;
+
+	// Unique object guid. Generated unique for a local process
+	FGuid ObjectGuid;
+
+	// Object Sync info
+	FTextureShareCoreObjectSync Sync;
+
+	// Local process description
+	FTextureShareCoreObjectProcessDesc ProcessDesc;
+
+public:
+	virtual ~FTextureShareCoreObjectDesc() = default;
+
+	virtual ITextureShareSerializeStream& Serialize(ITextureShareSerializeStream & Stream) override
+	{
+		return Stream << ShareName << ObjectGuid << Sync << ProcessDesc;
+	}
+
+public:
+	FTextureShareCoreObjectDesc() = default;
+
+	bool EqualsFunc(const FGuid& InObjectGuid) const
+	{
+		return ObjectGuid == InObjectGuid;
+	}
+
+	bool EqualsFunc(const FTextureShareCoreObjectDesc& InObjectDesc) const
+	{
+		return ObjectGuid == InObjectDesc.ObjectGuid;
+	}
+
+	bool operator==(const FTextureShareCoreObjectDesc& InObjectDesc) const
+	{
+		return ObjectGuid == InObjectDesc.ObjectGuid;
+	}
+};
diff --git a/Engine/Extras/VirtualProduction/TextureShare/TextureShareSDK/Source/Containers/TextureShareCoreContainers_ProxyData.h b/Engine/Extras/VirtualProduction/TextureShare/TextureShareSDK/Source/Containers/TextureShareCoreContainers_ProxyData.h
new file mode 100644
--- /dev/null
+++ b/Engine/Extras/VirtualProduction/TextureShare/TextureShareSDK/Source/Containers/TextureShareCoreContainers_ProxyData.h
@@ -0,0 +1,48 @@
+// Copyright Epic Games, Inc. All Rights Reserved.
+
+#pragma once
+#include "Containers/TextureShareCoreContainers_FrameMarker.h"
+#include "Containers/TextureShareCoreContainers_SceneData.h"
+#include "Containers/TextureShareCoreContainers_ResourceHandle.h"
+
+/**
+ * Proxy data of the TextureShareCore object
+ * The structure is serializable (binary compatible) and reflected on the SDK side (UE types will be replaced with simplified copies from the SDK).
+ */
+struct FTextureShareCoreProxyData
+	: public ITextureShareSerialize
+{
+	// Frame markers on local process (Copied from game thread value)
+	// Must be copied from the game thread object data immediately after BeginFrameSync_RenderThread.
+	// **Currently, only the UE process side makes changes to this property.**
+	FTextureShareCoreFrameMarker FrameMarker;
+
+	// Frame markers of all connected objects from game thread
+	// Must be copied from the game thread object data immediately after BeginFrameSync_RenderThread.
+	// **Currently, only the UE process side makes changes to this property.**
+	TArraySerializable<FTextureShareCoreObjectFrameMarker> RemoteFrameMarkers;
+
+	// Scene data (viewports, matrices, etc)
+	// **Currently, only the UE process side makes changes to this property.**
+	TArraySerializable<FTextureShareCoreSceneViewData> SceneData;
+
+	// Descriptors for shared resources. Change only from the owner of the resource.
+	// **Currently, only the UE process side makes changes to this property.**
+	TArraySerializable<FTextureShareCoreResourceHandle> ResourceHandles;
+
+public:
+	virtual ~FTextureShareCoreProxyData() = default;
+
+	virtual ITextureShareSerializeStream& Serialize(ITextureShareSerializeStream & Stream) override
+	{
+		return Stream << FrameMarker << RemoteFrameMarkers << SceneData << ResourceHandles;
+	}
+
+public:
+	void ResetProxyData()
+	{
+		RemoteFrameMarkers.Empty();
+		SceneData.Empty();
+		ResourceHandles.Empty();
+	}
+};
diff --git a/Engine/Extras/VirtualProduction/TextureShare/TextureShareSDK/Source/Containers/TextureShareCoreContainers_ResourceDesc.h b/Engine/Extras/VirtualProduction/TextureShare/TextureShareSDK/Source/Containers/TextureShareCoreContainers_ResourceDesc.h
new file mode 100644
--- /dev/null
+++ b/Engine/Extras/VirtualProduction/TextureShare/TextureShareSDK/Source/Containers/TextureShareCoreContainers_ResourceDesc.h
@@ -0,0 +1,77 @@
+// Copyright Epic Games, Inc. All Rights Reserved.
+
+#pragma once
+#include "Containers/TextureShareCoreContainers_ViewDesc.h"
+
+/**
+ * Resource descriptor
+ * The structure is serializable (binary compatible) and reflected on the SDK side (UE types will be replaced with simplified copies from the SDK).
+ */
+struct FTextureShareCoreResourceDesc
+	: public ITextureShareSerialize
+{
+	// Resource name
+	FString ResourceName;
+
+	// Resource owner
+	FTextureShareCoreViewDesc ViewDesc;
+
+	// Texture operation type (read or write)
+	ETextureShareTextureOp OperationType = ETextureShareTextureOp::Read;
+
+	// resource sync order for logic (default value '-1'unordered)
+	ETextureShareSyncStep SyncStep = ETextureShareSyncStep::Undefined;
+
+public:
+	virtual ~FTextureShareCoreResourceDesc() = default;
+
+	virtual ITextureShareSerializeStream& Serialize(ITextureShareSerializeStream & Stream) override
+	{
+		return Stream << ResourceName << ViewDesc << OperationType << SyncStep;
+	}
+
+public:
+	FTextureShareCoreResourceDesc(const FString& InResourceName, const FTextureShareCoreViewDesc& InViewDesc, const ETextureShareTextureOp InOperationType, const ETextureShareSyncStep InSyncStep)
+		: ResourceName(InResourceName), ViewDesc(InViewDesc), OperationType(InOperationType), SyncStep(InSyncStep)
+	{ }
+
+	FTextureShareCoreResourceDesc(const FString& InResourceName, const FTextureShareCoreViewDesc& InViewDesc, const ETextureShareTextureOp InOperationType)
+		: ResourceName(InResourceName), ViewDesc(InViewDesc), OperationType(InOperationType)
+	{ }
+
+	FTextureShareCoreResourceDesc(const FString& InResourceName, const ETextureShareTextureOp InOperationType, const ETextureShareSyncStep InSyncStep)
+		: ResourceName(InResourceName), OperationType(InOperationType), SyncStep(InSyncStep)
+	{ }
+
+	FTextureShareCoreResourceDesc(const FString& InResourceName, const ETextureShareTextureOp InOperationType)
+		: ResourceName(InResourceName), OperationType(InOperationType)
+	{ }
+
+	FTextureShareCoreResourceDesc(const FString& InResourceName)
+		: ResourceName(InResourceName)
+	{ }
+
+	FTextureShareCoreResourceDesc() = default;
+
+public:
+	bool EqualsFunc(const FTextureShareCoreViewDesc& InViewDesc) const
+	{
+		return ViewDesc.EqualsFunc(InViewDesc);
+	}
+
+	bool EqualsFunc(const FTextureShareCoreResourceDesc& InResourceDesc) const
+	{
+		return ViewDesc.EqualsFunc(InResourceDesc.ViewDesc)
+			&& ResourceName == InResourceDesc.ResourceName
+			&& (OperationType == InResourceDesc.OperationType
+				|| OperationType == ETextureShareTextureOp::Undefined
+				|| InResourceDesc.OperationType == ETextureShareTextureOp::Undefined);
+	}
+
+	bool operator==(const FTextureShareCoreResourceDesc& InResourceDesc) const
+	{
+		return ViewDesc == InResourceDesc.ViewDesc
+			&& ResourceName == InResourceDesc.ResourceName
+			&& OperationType == InResourceDesc.OperationType;
+	}
+};
diff --git a/Engine/Extras/VirtualProduction/TextureShare/TextureShareSDK/Source/Containers/TextureShareCoreContainers_ResourceHandle.h b/Engine/Extras/VirtualProduction/TextureShare/TextureShareSDK/Source/Containers/TextureShareCoreContainers_ResourceHandle.h
new file mode 100644
--- /dev/null
+++ b/Engine/Extras/VirtualProduction/TextureShare/TextureShareSDK/Source/Containers/TextureShareCoreContainers_ResourceHandle.h
@@ -0,0 +1,55 @@
+// Copyright Epic Games, Inc. All Rights Reserved.
+
+#pragma once
+#include "Containers/TextureShareCoreContainers_ResourceDesc.h"
+
+/**
+ * Resource handle data (Shared resource handlers and other data)
+ * The structure is serializable (binary compatible) and reflected on the SDK side (UE types will be replaced with simplified copies from the SDK).
+ */
+struct FTextureShareCoreResourceHandle
+	: public ITextureShareSerialize
+{
+	// Resource info
+	FTextureShareCoreResourceDesc ResourceDesc;
+
+	// GPU-GPU share purpose (this resource use NT handle)
+	Windows::HANDLE NTHandle = nullptr;
+
+	// GPU-GPU share purpose (not NT handle)
+	Windows::HANDLE SharedHandle = nullptr;
+
+	// Unique handle name
+	FGuid SharedHandleGuid;
+
+public:
+	virtual ~FTextureShareCoreResourceHandle() = default;
+
+	virtual ITextureShareSerializeStream& Serialize(ITextureShareSerializeStream & Stream) override
+	{
+		return Stream << ResourceDesc << NTHandle << SharedHandle << SharedHandleGuid;
+	}
+
+public:
+	bool HandleEquals(const FTextureShareCoreResourceHandle& In) const
+	{
+		return NTHandle == In.NTHandle
+			&& SharedHandleGuid == In.SharedHandleGuid
+			&& SharedHandle == In.SharedHandle;
+	}
+
+	bool EqualsFunc(const FTextureShareCoreResourceDesc& InResourceDesc) const
+	{
+		return ResourceDesc.EqualsFunc(InResourceDesc);
+	}
+
+	bool EqualsFunc(const FTextureShareCoreViewDesc& InViewDesc) const
+	{
+		return ResourceDesc.EqualsFunc(InViewDesc);
+	}
+
+	bool operator==(const FTextureShareCoreResourceHandle& In) const
+	{
+		return ResourceDesc == In.ResourceDesc;
+	}
+};
diff --git a/Engine/Extras/VirtualProduction/TextureShare/TextureShareSDK/Source/Containers/TextureShareCoreContainers_ResourceRequest.h b/Engine/Extras/VirtualProduction/TextureShare/TextureShareSDK/Source/Containers/TextureShareCoreContainers_ResourceRequest.h
new file mode 100644
--- /dev/null
+++ b/Engine/Extras/VirtualProduction/TextureShare/TextureShareSDK/Source/Containers/TextureShareCoreContainers_ResourceRequest.h
@@ -0,0 +1,80 @@
+// Copyright Epic Games, Inc. All Rights Reserved.
+
+#pragma once
+#include "Containers/TextureShareCoreContainers_ResourceDesc.h"
+
+/**
+ * Resource request data
+ * When no size and/or format is specified, values from another process or values from local resources are used.
+ * The structure is serializable (binary compatible) and reflected on the SDK side (UE types will be replaced with simplified copies from the SDK).
+ */
+struct FTextureShareCoreResourceRequest
+	: public ITextureShareSerialize
+{
+	// Resource info
+	FTextureShareCoreResourceDesc ResourceDesc;
+
+	// MGPU support. When a texture is rendered on a GPU other than the destination,
+	// it must be transferred between GPUs.
+	// The transfer is performed on the UE side in the TextureShare module.
+	int32 GPUIndex = -1;
+
+	// Required texture format
+	// The UE process only uses the "EPixelFormat" value.
+	// Otherwise, find the best of "EPixelFormat" associated with "DXGI_FORMAT".
+	EPixelFormat PixelFormat = EPixelFormat::PF_Unknown;
+	DXGI_FORMAT       Format = DXGI_FORMAT_UNKNOWN;
+
+	// Required texture size (or zero if the original value is acceptable)
+	FIntPoint           Size = FIntPoint::ZeroValue;
+
+	// Experimental: nummips feature
+	uint32 NumMips = 0;
+
+public:
+	virtual ~FTextureShareCoreResourceRequest() = default;
+
+	virtual ITextureShareSerializeStream& Serialize(ITextureShareSerializeStream & Stream) override
+	{
+		return Stream << ResourceDesc << GPUIndex << PixelFormat << Format << Size << NumMips;
+	}
+
+public:
+	FTextureShareCoreResourceRequest() = default;
+
+	FTextureShareCoreResourceRequest(const FTextureShareCoreResourceDesc& InResourceDesc)
+		: ResourceDesc(InResourceDesc)
+	{ }
+
+	FTextureShareCoreResourceRequest(const FTextureShareCoreResourceDesc& InResourceDesc, const FIntPoint& InSize, const DXGI_FORMAT InFormat = DXGI_FORMAT_UNKNOWN, const EPixelFormat InPixelFormat = EPixelFormat::PF_Unknown)
+		: ResourceDesc(InResourceDesc)
+		, PixelFormat(InPixelFormat)
+		, Format(InFormat)
+		, Size(InSize)
+	{ }
+
+	FTextureShareCoreResourceRequest(const FTextureShareCoreResourceDesc& InResourceDesc, const DXGI_FORMAT InFormat)
+		: ResourceDesc(InResourceDesc)
+		, Format(InFormat)
+	{ }
+
+	void SetPixelFormat(uint32 InPixelFormat)
+	{
+		PixelFormat = (EPixelFormat)InPixelFormat;
+	}
+
+	bool EqualsFunc(const FTextureShareCoreResourceDesc& InResourceDesc) const
+	{
+		return ResourceDesc.EqualsFunc(InResourceDesc);
+	}
+
+	bool EqualsFunc(const FTextureShareCoreViewDesc& InViewDesc) const
+	{
+		return ResourceDesc.EqualsFunc(InViewDesc);
+	}
+
+	bool operator==(const FTextureShareCoreResourceRequest& InResourceRequest) const
+	{
+		return ResourceDesc == InResourceRequest.ResourceDesc;
+	}
+};
diff --git a/Engine/Extras/VirtualProduction/TextureShare/TextureShareSDK/Source/Containers/TextureShareCoreContainers_SceneData.h b/Engine/Extras/VirtualProduction/TextureShare/TextureShareSDK/Source/Containers/TextureShareCoreContainers_SceneData.h
new file mode 100644
--- /dev/null
+++ b/Engine/Extras/VirtualProduction/TextureShare/TextureShareSDK/Source/Containers/TextureShareCoreContainers_SceneData.h
@@ -0,0 +1,202 @@
+// Copyright Epic Games, Inc. All Rights Reserved.
+
+#pragma once
+#include "Containers/TextureShareCoreContainers_ViewDesc.h"
+
+/**
+ * Scene data: View matrices
+ * The structure is serializable (binary compatible) and reflected on the SDK side (UE types will be replaced with simplified copies from the SDK).
+ */
+struct FTextureShareCoreSceneViewMatrices
+	: public ITextureShareSerialize
+{
+	// ViewToClip : UE projection matrix projects such that clip space Z=1 is the near plane, and Z=0 is the infinite far plane.
+	FMatrix ProjectionMatrix;
+
+	// ViewToClipNoAA : UE projection matrix projects such that clip space Z=1 is the near plane, and Z=0 is the infinite far plane. Don't apply any AA jitter
+	FMatrix ProjectionNoAAMatrix;
+
+	// WorldToView..
+	FMatrix ViewMatrix;
+
+	// WorldToClip : UE projection matrix projects such that clip space Z=1 is the near plane, and Z=0 is the infinite far plane. */
+	FMatrix ViewProjectionMatrix;
+
+	// The view-projection transform, starting from world-space points translated by -ViewOrigin.
+	FMatrix TranslatedViewProjectionMatrix;
+
+	// The translation to apply to the world before TranslatedViewProjectionMatrix. Usually it is -ViewOrigin but with rereflections this can differ
+	FVector PreViewTranslation;
+
+	// To support ortho and other modes this is redundant, in world space
+	FVector ViewOrigin;
+
+	// Scale applied by the projection matrix in X and Y.
+	FVector2D ProjectionScale;
+
+	// TemporalAA jitter offset currently stored in the projection matrix
+	FVector2D TemporalAAProjectionJitter;
+
+	// Scale factor to use when computing the size of a sphere in pixels.
+	float ScreenScale = 1.f;
+
+public:
+	virtual ~FTextureShareCoreSceneViewMatrices() = default;
+
+	virtual ITextureShareSerializeStream& Serialize(ITextureShareSerializeStream & Stream) override
+	{
+		return Stream << ProjectionMatrix << ProjectionNoAAMatrix
+			<< ViewMatrix << ViewProjectionMatrix << TranslatedViewProjectionMatrix << PreViewTranslation
+			<< ViewOrigin << ProjectionScale << TemporalAAProjectionJitter << ScreenScale;
+	}
+};
+
+/**
+ * Scene data: View
+ * The structure is serializable (binary compatible) and reflected on the SDK side (UE types will be replaced with simplified copies from the SDK).
+ */
+struct FTextureShareCoreSceneView
+	: public ITextureShareSerialize
+{
+	FTextureShareCoreSceneViewMatrices ViewMatrices;
+
+	// Final position of the view in the final render target (in pixels), potentially constrained by an aspect ratio requirement (black bars)
+	FIntRect UnscaledViewRect;
+
+	// Raw view size (in pixels), used for screen space calculations
+	FIntRect UnconstrainedViewRect;
+
+	// Variables used to determine the view matrix
+	FVector  ViewLocation;
+	FRotator ViewRotation;
+	FQuat    BaseHmdOrientation;
+	FVector  BaseHmdLocation;
+	float    WorldToMetersScale = 1.f;
+
+	// For stereoscopic rendering, unique index identifying the view across view families
+	int32 StereoViewIndex;
+
+	// For stereoscopic rendering, view family index of the primary view associated with this view
+	int32 PrimaryViewIndex;
+
+	// Actual field of view and that desired by the camera originally
+	float FOV;
+	float DesiredFOV;
+
+public:
+	virtual ~FTextureShareCoreSceneView() = default;
+
+	virtual ITextureShareSerializeStream& Serialize(ITextureShareSerializeStream & Stream) override
+	{
+		return Stream << ViewMatrices
+			<< UnscaledViewRect << UnconstrainedViewRect
+			<< ViewLocation << ViewRotation << BaseHmdOrientation << BaseHmdLocation << WorldToMetersScale
+			<< StereoViewIndex << PrimaryViewIndex << FOV << DesiredFOV;
+	}
+};
+
+/**
+ * Scene data: GameTime
+ * The structure is serializable (binary compatible) and reflected on the SDK side (UE types will be replaced with simplified copies from the SDK).
+ */
+struct FTextureShareCoreSceneGameTime
+	: public ITextureShareSerialize
+{
+	float RealTimeSeconds;
+	float WorldTimeSeconds;
+
+	float DeltaRealTimeSeconds;
+	float DeltaWorldTimeSeconds;
+
+public:
+	virtual ~FTextureShareCoreSceneGameTime() = default;
+
+	virtual ITextureShareSerializeStream& Serialize(ITextureShareSerializeStream & Stream) override
+	{
+		return Stream << RealTimeSeconds << WorldTimeSeconds
+			<< DeltaRealTimeSeconds << DeltaWorldTimeSeconds;
+	}
+};
+
+/**
+ * Scene data: ViewFamily
+ * The structure is serializable (binary compatible) and reflected on the SDK side (UE types will be replaced with simplified copies from the SDK).
+ */
+struct FTextureShareCoreSceneViewFamily
+	: public ITextureShareSerialize
+{
+	// The current time
+	FTextureShareCoreSceneGameTime GameTime;
+
+	// Copy from main thread GFrameNumber to be accessible on render thread side. UINT_MAX before CreateSceneRenderer() or BeginRenderingViewFamily() was called
+	uint32 FrameNumber = 0;
+
+	// When enabled, the post processing will output in HDR space
+	bool bIsHDR = false;
+
+	// Gamma correction used when rendering this family. Default is 1.0
+	float GammaCorrection = 1.f;
+
+	// Secondary view fraction to support High DPI monitor
+	float SecondaryViewFraction = 1.f;
+
+public:
+	virtual ~FTextureShareCoreSceneViewFamily() = default;
+
+	virtual ITextureShareSerializeStream& Serialize(ITextureShareSerializeStream & Stream) override
+	{
+		return Stream << GameTime << FrameNumber << bIsHDR << GammaCorrection << SecondaryViewFraction;
+	}
+};
+
+/**
+ * Scene data
+ * The structure is serializable (binary compatible) and reflected on the SDK side (UE types will be replaced with simplified copies from the SDK).
+ */
+struct FTextureShareCoreSceneViewData
+	: public ITextureShareSerialize
+{
+	// Eye type of this resource (support stereo)
+	FTextureShareCoreViewDesc ViewDesc;
+
+	// Scene view data
+	FTextureShareCoreSceneView View;
+
+	// Scene ViewFamily data
+	FTextureShareCoreSceneViewFamily ViewFamily;
+
+public:
+	virtual ~FTextureShareCoreSceneViewData() = default;
+
+	virtual ITextureShareSerializeStream& Serialize(ITextureShareSerializeStream & Stream) override
+	{
+		return Stream << ViewDesc << View << ViewFamily;
+	}
+
+public:
+	FTextureShareCoreSceneViewData() = default;
+
+	FTextureShareCoreSceneViewData(const FString& InViewId, const ETextureShareEyeType InEyeType = ETextureShareEyeType::Default)
+		: ViewDesc(InViewId, InEyeType)
+	{ }
+
+	FTextureShareCoreSceneViewData(const FTextureShareCoreViewDesc& InViewDesc)
+		: ViewDesc(InViewDesc)
+	{ }
+
+public:
+	bool EqualsFunc(const FTextureShareCoreSceneViewData& InViewData) const
+	{
+		return ViewDesc.EqualsFunc(InViewData.ViewDesc);
+	}
+
+	bool EqualsFunc(const FTextureShareCoreViewDesc& InViewDesc) const
+	{
+		return ViewDesc.EqualsFunc(InViewDesc);
+	}
+
+	bool operator==(const FTextureShareCoreSceneViewData& InSceneViewData) const
+	{
+		return ViewDesc == InSceneViewData.ViewDesc;
+	}
+};
diff --git a/Engine/Extras/VirtualProduction/TextureShare/TextureShareSDK/Source/Containers/TextureShareCoreContainers_SyncSettings.h b/Engine/Extras/VirtualProduction/TextureShare/TextureShareSDK/Source/Containers/TextureShareCoreContainers_SyncSettings.h
new file mode 100644
--- /dev/null
+++ b/Engine/Extras/VirtualProduction/TextureShare/TextureShareSDK/Source/Containers/TextureShareCoreContainers_SyncSettings.h
@@ -0,0 +1,164 @@
+// Copyright Epic Games, Inc. All Rights Reserved.
+
+#pragma once
+#include "Serialize/TextureShareCoreSerialize.h"
+
+/**
+ * Connection state
+ */
+enum class ETextureShareCoreFrameConnectionsState : int8
+{
+	SkipFrame = 0,
+	Wait,
+	Accept
+};
+
+/**
+ * SyncSettings: Connection - Wait until the required number of processes are connected at new frame
+ * The structure is serializable (binary compatible) and reflected on the SDK side (UE types will be replaced with simplified copies from the SDK).
+ */
+struct FTextureShareCoreFrameConnectionsSettings
+	: public ITextureShareSerialize
+{
+	// Required number of processes. The TS frame will only start when this number is reached.
+	int32 MinValue = 0;
+
+	// Connect only to processes with names from this list
+	TArraySerializable<FString> AllowedProcessNames;
+
+	// Ignore processes with names from this list
+	TArraySerializable<FString> BannedProcessNames;
+
+public:
+	virtual ~FTextureShareCoreFrameConnectionsSettings() = default;
+
+	virtual ITextureShareSerializeStream& Serialize(ITextureShareSerializeStream& Stream) override
+	{
+		return Stream << MinValue << AllowedProcessNames << BannedProcessNames;
+	}
+
+public:
+	ETextureShareCoreFrameConnectionsState GetConnectionsState(const int32 InReadyToConnectObjectsCount, const int32 InFrameConnectionsCount) const
+	{
+		// in case no TS processes and min=0, skip frame immediately
+		if (MinValue < 1 && InFrameConnectionsCount < 1)
+		{
+			return ETextureShareCoreFrameConnectionsState::SkipFrame;
+		}
+
+		if (InFrameConnectionsCount == InReadyToConnectObjectsCount)
+		{
+			if (MinValue < InReadyToConnectObjectsCount)
+			{
+				return ETextureShareCoreFrameConnectionsState::Accept;
+			}
+		}
+
+		return ETextureShareCoreFrameConnectionsState::Wait;
+	}
+};
+
+/**
+ * SyncSettings: define all sync steps for this object
+ * The structure is serializable (binary compatible) and reflected on the SDK side (UE types will be replaced with simplified copies from the SDK).
+ */
+struct FTextureShareCoreFrameSyncSettings
+	: public ITextureShareSerialize
+{
+	// Used sync steps by this process (handled by logic above)
+	// Order of values: grow
+	TArraySerializable<ETextureShareSyncStep> Steps;
+
+public:
+	virtual ~FTextureShareCoreFrameSyncSettings() = default;
+
+	virtual ITextureShareSerializeStream& Serialize(ITextureShareSerializeStream& Stream) override
+	{
+		return Stream << Steps;
+	}
+
+public:
+	bool Contains(const ETextureShareSyncStep InSyncStep) const
+	{
+		return Steps.Find(InSyncStep) != INDEX_NONE;
+	}
+
+	void Append(const TArray<ETextureShareSyncStep>& InSyncSteps)
+	{
+		Steps.AppendSorted(InSyncSteps);
+	}
+
+	bool EqualsFunc(const FTextureShareCoreFrameSyncSettings& In) const
+	{
+		return Steps.EqualsFunc(In.Steps);
+	}
+};
+
+/**
+ * SyncSettings: TimeOut values
+ * The structure is serializable (binary compatible) and reflected on the SDK side (UE types will be replaced with simplified copies from the SDK).
+ */
+struct FTextureShareCoreTimeOutSettings
+	: public ITextureShareSerialize
+{
+	// Memory mutex timeout in miliseconds
+	uint32 MemoryMutexTimeout = 1000;
+
+	// Timeout for frame begin (process
+	// 0 - infinite
+	uint32 FrameBeginTimeOut = 0;
+
+	// Max frame sync time, until this frame lost.
+	// 0 - infinite
+	uint32 FrameSyncTimeOut = 1000;
+
+	// Split time for timeouts. prevent freeze
+	uint32 FrameBeginTimeOutSplit = 100;
+	uint32 FrameSyncTimeOutSplit = 100;
+
+	// now sync logic threat object as lost, if last access time expired.
+	// If the process is alive, it is re-created on-demant and re-connected
+	// '-1' - disable this option
+	// This value is unique per process.
+	int32 ProcessLostStatusTimeOut = 2000;
+
+	// Multithreaded mutex timeout used in FTextureShareCoreObject::LockThread Mutex()
+	int32 ThreadMutexTimeout = 1500;
+
+public:
+	virtual ~FTextureShareCoreTimeOutSettings() = default;
+
+	virtual ITextureShareSerializeStream& Serialize(ITextureShareSerializeStream & Stream) override
+	{
+		return Stream << MemoryMutexTimeout << FrameBeginTimeOut << FrameSyncTimeOut
+			<< FrameBeginTimeOutSplit << FrameSyncTimeOutSplit
+			<< ProcessLostStatusTimeOut
+			<< ThreadMutexTimeout;
+	}
+};
+
+/**
+ * SyncSettings for TextureShareCore object
+ * The structure is serializable (binary compatible) and reflected on the SDK side (UE types will be replaced with simplified copies from the SDK).
+ */
+struct FTextureShareCoreSyncSettings
+	: public ITextureShareSerialize
+{
+	// At the beginning of each frame, the local process waits for the required number of processes.
+	// If it is not reached within the specified time, then this frame will not start.
+	FTextureShareCoreFrameConnectionsSettings FrameConnectionSettings;
+
+	// Timeout settings for internal IPC sync logic
+	FTextureShareCoreTimeOutSettings TimeoutSettings;
+
+	// Frame sync logic settings
+	FTextureShareCoreFrameSyncSettings FrameSyncSettings;
+
+public:
+	virtual ~FTextureShareCoreSyncSettings() = default;
+
+	virtual ITextureShareSerializeStream& Serialize(ITextureShareSerializeStream & Stream) override
+	{
+		return Stream << FrameConnectionSettings << TimeoutSettings << FrameSyncSettings;
+	}
+};
diff --git a/Engine/Extras/VirtualProduction/TextureShare/TextureShareSDK/Source/Containers/TextureShareCoreContainers_ViewDesc.h b/Engine/Extras/VirtualProduction/TextureShare/TextureShareSDK/Source/Containers/TextureShareCoreContainers_ViewDesc.h
new file mode 100644
--- /dev/null
+++ b/Engine/Extras/VirtualProduction/TextureShare/TextureShareSDK/Source/Containers/TextureShareCoreContainers_ViewDesc.h
@@ -0,0 +1,57 @@
+// Copyright Epic Games, Inc. All Rights Reserved.
+
+#pragma once
+#include "Serialize/TextureShareCoreSerialize.h"
+#include "Misc/TextureShareCoreStrings.h"
+
+/**
+ * View descriptor (Name + EyeType)
+ * The structure is serializable (binary compatible) and reflected on the SDK side (UE types will be replaced with simplified copies from the SDK).
+ */
+struct FTextureShareCoreViewDesc
+	: public ITextureShareSerialize
+{
+	// View name
+	FString Id = TextureShareCoreStrings::Default::ViewId;
+
+	// Eye type of this view (support stereo)
+	ETextureShareEyeType EyeType = ETextureShareEyeType::Default;
+
+public:
+	virtual ~FTextureShareCoreViewDesc() = default;
+
+	virtual ITextureShareSerializeStream& Serialize(ITextureShareSerializeStream & Stream) override
+	{
+		return Stream << Id << EyeType;
+	}
+
+public:
+	FTextureShareCoreViewDesc() = default;
+
+	FTextureShareCoreViewDesc(const FTextureShareCoreViewDesc & In)
+		: Id(In.Id), EyeType(In.EyeType)
+	{ }
+
+	FTextureShareCoreViewDesc(const FString & InViewId, const ETextureShareEyeType InEyeType = ETextureShareEyeType::Default)
+		: Id(InViewId), EyeType(InEyeType)
+	{ }
+
+	FTextureShareCoreViewDesc(const ETextureShareEyeType InEyeType)
+		: EyeType(InEyeType)
+	{ }
+
+	bool EqualsFunc(const FString& InViewId) const
+	{
+		return InViewId == Id;
+	}
+
+	bool EqualsFunc(const FTextureShareCoreViewDesc & InViewDesc) const
+	{
+		return InViewDesc.Id == Id && EyeType == InViewDesc.EyeType;
+	}
+
+	bool operator==(const FTextureShareCoreViewDesc& InViewDesc) const
+	{
+		return InViewDesc.Id == Id && EyeType == InViewDesc.EyeType;
+	}
+};
diff --git a/Engine/Extras/VirtualProduction/TextureShare/TextureShareSDK/Source/Containers/TextureShareCoreEnums.h b/Engine/Extras/VirtualProduction/TextureShare/TextureShareSDK/Source/Containers/TextureShareCoreEnums.h
new file mode 100644
--- /dev/null
+++ b/Engine/Extras/VirtualProduction/TextureShare/TextureShareSDK/Source/Containers/TextureShareCoreEnums.h
@@ -0,0 +1,232 @@
+// Copyright Epic Games, Inc. All Rights Reserved.
+
+#pragma once
+
+#if __UNREAL__
+#include "CoreMinimal.h"
+#endif
+
+/**
+ * Render Device interface type
+ */
+enum class ETextureShareDeviceType : uint8
+{
+	Undefined = 0,
+
+	D3D11,
+	D3D12,
+
+	Vulkan
+};
+
+/**
+ * Process type
+ */
+enum class ETextureShareProcessType : uint8
+{
+	Undefined = 0,
+
+	// Unreal Engine base process
+	UE,
+
+	// TextureShare SDK process
+	SDK,
+
+	// Marked as Invalid process
+	Invalid
+};
+
+/**
+ * Multithread mutex index
+ * Using this enum, you can access multi-threaded mutexes.
+ */
+enum class ETextureShareThreadMutex : uint8
+{
+	GameThread = 0,
+	RenderingThread,
+
+	COUNT
+};
+
+/**
+ * Texture operation
+ */
+enum class ETextureShareTextureOp : uint8
+{
+	// Send texture to remote process
+	Write = 0,
+
+	// Receive texture
+	Read,
+
+	// Operation not defined (useful to find any type of operation)
+	Undefined,
+};
+
+/**
+ * View eye type
+ */
+enum class ETextureShareEyeType : uint8
+{
+	// Monoscopic
+	Default = 0,
+
+	// Stereoscopic rendering
+	StereoLeft,
+	StereoRight
+};
+
+/**
+ * Sync steps template type
+ */
+enum class ETextureShareFrameSyncTemplate : int8
+{
+	// Sync logic for BP object
+	Default = 0,
+
+	// Minimalistic sync logic for SDK object
+	SDK,
+
+	// Sync logic for DC object
+	DisplayCluster,
+};
+
+/**
+ * Sync steps values
+ */
+enum class ETextureShareSyncStep : int8
+{
+	InterprocessConnection = 0,
+	
+	/**
+	 * Frame sync steps (GameThread)
+	 */
+
+	FrameBegin,
+
+	FramePreSetupBegin,
+	FramePreSetupEnd,
+
+	FrameSetupBegin,
+	FrameSetupEnd,
+
+	FramePostSetupBegin,
+	FramePostSetupEnd,
+
+	FrameEnd,
+
+	/**
+	 * FrameProxy sync steps (RenderThread)
+	 */
+
+	FrameProxyBegin,
+
+	FrameSceneFinalColorBegin,
+	FrameSceneFinalColorEnd,
+
+	FrameProxyPreRenderBegin,
+	FrameProxyPreRenderEnd,
+
+	FrameProxyRenderBegin,
+	FrameProxyRenderEnd,
+
+	FrameProxyPostWarpBegin,
+	FrameProxyPostWarpEnd,
+
+	FrameProxyPostRenderBegin,
+	FrameProxyPostRenderEnd,
+
+	FrameProxyBackBufferReadyToPresentBegin,
+	FrameProxyBackBufferReadyToPresentEnd,
+
+	FrameProxyFlush,
+
+	FrameProxyEnd,
+
+	/**
+	 * Special values
+	 */
+	COUNT,
+	Undefined = -1,
+};
+
+/**
+ * Object sync barrier pass
+ */
+enum class ETextureShareSyncPass : int8
+{
+	Undefined = 0,
+
+	Enter,
+	Exit,
+	Complete,
+
+	COUNT
+};
+
+/**
+ * Object sync barrier state
+ */
+enum class ETextureShareSyncState : int8
+{
+	// After sync reset
+	Undefined = 0,
+
+	Enter,
+	EnterCompleted,
+
+	Exit,
+	ExitCompleted,
+
+	Completed,
+};
+
+/**
+ * Projection matrix source type
+ */
+enum class ETextureShareCoreSceneViewManualProjectionType : uint8
+{
+	FrustumAngles = 0,
+	Matrix
+};
+
+/**
+ * Rotator operation type
+ * How to change the projection view rotator value
+ */
+enum class ETextureShareViewRotationDataType : uint8
+{
+	// Use original value
+	Original = 0,
+
+	// Add value from manual projection
+	Relative,
+
+	// Use value from manual projection
+	Absolute
+};
+
+/**
+ * Location operation type
+ * How to change the projection view location value
+ */
+enum class ETextureShareViewLocationDataType : uint8
+{
+	// Use original value
+	Original = 0,
+
+	// Add value from manual projection
+	Relative,
+
+	// Use value from manual projection
+	Absolute,
+
+	// Use original value, but without stereo view offset
+	Original_NoViewOffset,
+
+	// Add value from manual projection, but without stereo view offset
+	Relative_NoViewOffset,
+
+	// Use value from manual projection, but without stereo view offset
+	Absolute_NoViewOffset
+};
diff --git a/Engine/Extras/VirtualProduction/TextureShare/TextureShareSDK/Source/Containers/TextureShareSDKContainers.h b/Engine/Extras/VirtualProduction/TextureShare/TextureShareSDK/Source/Containers/TextureShareSDKContainers.h
new file mode 100644
--- /dev/null
+++ b/Engine/Extras/VirtualProduction/TextureShare/TextureShareSDK/Source/Containers/TextureShareSDKContainers.h
@@ -0,0 +1,22 @@
+// Copyright Epic Games, Inc. All Rights Reserved.
+
+#pragma once
+#include "Containers/UnrealEngine/TextureShareSDKUnrealEngineArray.h"
+#include "Containers/UnrealEngine/TextureShareSDKUnrealEngineContainers.h"
+#include "Containers/UnrealEngine/TextureShareSDKUnrealEngineEnums_PixelFormat.h"
+#include "Containers/UnrealEngine/TextureShareSDKUnrealEngineMath.h"
+#include "Containers/UnrealEngine/TextureShareSDKUnrealEngineString.h"
+
+#include "Containers/TextureShareCoreEnums.h"
+#include "Containers/TextureShareCoreContainers.h"
+
+#include "Containers/TextureShareCoreContainers_DeviceD3D11.h"
+#include "Containers/TextureShareCoreContainers_DeviceD3D12.h"
+#include "Containers/TextureShareCoreContainers_DeviceVulkan.h"
+
+#ifndef TEXTURESHARESDK_API
+
+// Support SDK user side DLL functions import
+#define TEXTURESHARESDK_API __declspec(dllimport)
+
+#endif
diff --git a/Engine/Extras/VirtualProduction/TextureShare/TextureShareSDK/Source/Containers/UnrealEngine/TextureShareSDKUnrealEngineArray.h b/Engine/Extras/VirtualProduction/TextureShare/TextureShareSDK/Source/Containers/UnrealEngine/TextureShareSDKUnrealEngineArray.h
new file mode 100644
--- /dev/null
+++ b/Engine/Extras/VirtualProduction/TextureShare/TextureShareSDK/Source/Containers/UnrealEngine/TextureShareSDKUnrealEngineArray.h
@@ -0,0 +1,211 @@
+// Copyright Epic Games, Inc. All Rights Reserved.
+
+#pragma once
+#include "Containers/UnrealEngine/TextureShareSDKUnrealEngineTypes.h"
+
+#ifndef __UNREAL__
+
+#include <vector>
+
+enum { INDEX_NONE = -1 };
+
+//@todo: implement better logic for std::vector vs TArray (this is too simplistic now)
+
+/**
+ * Implementation of TArray on the SDK ext app side (reflect this data type)
+ */
+template<typename InElementType>
+struct TArray
+	: public std::vector<InElementType>
+{
+	typedef InElementType ElementType;
+	typedef         int32 SizeType;
+
+	/**
+	 * Finds element within the array.
+	 *
+	 * @param Item Item to look for.
+	 * @returns Index of the found element. INDEX_NONE otherwise.
+	 * @see FindLast, FindLastByPredicate
+	 */
+	inline SizeType Find(const ElementType& Item) const
+	{
+		const SizeType Count = Num();
+
+		for(int Index = 0; Index < Count; Index++)
+		{
+			if (this->operator[](Index) == Item)
+			{
+				return Index;
+			}
+		}
+		return INDEX_NONE;
+	}
+
+	/**
+	 * Helper function for returning a typed pointer to the first array entry.
+	 *
+	 * @returns Pointer to first array entry or nullptr if ArrayMax == 0.
+	 */
+	inline ElementType* GetData()
+	{
+		return (ElementType*)this->data();
+	}
+
+	/**
+	 * Helper function for returning a typed pointer to the first array entry.
+	 *
+	 * @returns Pointer to first array entry or nullptr if ArrayMax == 0.
+	 */
+	inline const ElementType* GetData() const
+	{
+		return (const ElementType*)this->data();
+	}
+
+	/**
+	 * Empties the array. It calls the destructors on held items if needed.
+	 */
+	inline void Empty()
+	{
+		this->clear();
+	}
+
+	/**
+	 * Same as empty, but doesn't change memory allocations, unless the new size is larger than
+	 * the current array. It calls the destructors on held items if needed and then zeros the ArrayNum.
+	 */
+	void Reset()
+	{
+		this->clear();
+	}
+
+	/**
+	 * Returns true if the array is empty and contains no elements.
+	 *
+	 * @returns True if the array is empty.
+	 * @see Num
+	 */
+	inline bool IsEmpty() const
+	{
+		return Num() == 0;
+	}
+
+	/**
+	 * Returns number of elements in array.
+	 *
+	 * @returns Number of elements in array.
+	 */
+	inline SizeType Num() const
+	{
+		const SizeType Index = (SizeType)(this->size());
+
+		return Index;
+	}
+
+	/**
+	 * Adds a new item to the end of the array, possibly reallocating the whole array to fit.
+	 *
+	 * @param Item The item to add
+	 * @return Index to the new item
+	 * @see AddDefaulted, AddUnique, AddZeroed, Append, Insert
+	 */
+	inline SizeType Add(const ElementType& Item)
+	{
+		const SizeType Index = Num();
+
+		this->push_back(Item);
+
+		return Index;
+	}
+
+	inline SizeType AddUnique(const ElementType& Item)
+	{
+		SizeType ItemIndex = Find(Item);
+		if (ItemIndex == INDEX_NONE)
+		{
+			// Add unique
+			ItemIndex = Num();
+			this->push_back(Item);
+		}
+
+		return ItemIndex;
+	}
+
+	inline void Insert(const ElementType& Item, const int32 Index)
+	{
+		this->insert(this->begin() + Index, Item);
+	}
+
+	/**
+	 * Adds new items to the end of the array, possibly reallocating the whole
+	 * array to fit. The new items will be zeroed.
+	 *
+	 * Caution, AddZeroed() will create elements without calling the
+	 * constructor and this is not appropriate for element types that require
+	 * a constructor to function properly.
+	 *
+	 * @param  Count  The number of new items to add.
+	 * @return Index to the first of the new items.
+	 * @see Add, AddDefaulted, AddUnique, Append, Insert
+	 */
+	inline SizeType AddZeroed(SizeType Count = 1)
+	{
+		if (Count > 0)
+		{
+			const SizeType Index = Num();
+
+			this->resize(Index + Count);
+
+			// set to zero all new elements
+			std::memset((uint8*)(GetData()) + sizeof(ElementType) * Index, 0, sizeof(ElementType) * Count);
+
+			return Index;
+		}
+
+		return INDEX_NONE;
+	}
+
+	/**
+	 * Adds new items to the end of the array, possibly reallocating the whole
+	 * array to fit. The new items will be default-constructed.
+	 *
+	 * @param  Count  The number of new items to add.
+	 * @return Index to the first of the new items.
+	 * @see Add, AddZeroed, AddUnique, Append, Insert
+	 */
+	inline SizeType AddDefaulted(SizeType Count = 1)
+	{
+		if (Count > 0)
+		{
+			const SizeType Index = Num();
+
+			this->resize(Index + Count);
+
+			return Index;
+		}
+
+		return INDEX_NONE;
+	}
+
+	/**
+	 * Removes an element (or elements) at given location optionally shrinking
+	 * the array.
+	 *
+	 * @param Index Location in array of the element to remove.
+	 */
+	inline void RemoveAt(SizeType Index, SizeType Count = 1)
+	{
+		if (Index >= 0 && Index < Num())
+		{
+			if (Count > 1)
+			{
+				this->erase(this->begin() + Index, this->begin() + Index + Count);
+			}
+			else
+			{
+				this->erase(this->begin() + Index);
+			}
+		}
+	}
+};
+#endif
diff --git a/Engine/Extras/VirtualProduction/TextureShare/TextureShareSDK/Source/Containers/UnrealEngine/TextureShareSDKUnrealEngineContainers.h b/Engine/Extras/VirtualProduction/TextureShare/TextureShareSDK/Source/Containers/UnrealEngine/TextureShareSDKUnrealEngineContainers.h
new file mode 100644
--- /dev/null
+++ b/Engine/Extras/VirtualProduction/TextureShare/TextureShareSDK/Source/Containers/UnrealEngine/TextureShareSDKUnrealEngineContainers.h
@@ -0,0 +1,150 @@
+// Copyright Epic Games, Inc. All Rights Reserved.
+
+#pragma once
+#include "Containers/UnrealEngine/TextureShareSDKUnrealEngineTypes.h"
+
+#ifndef __UNREAL__
+/**
+ * Implementation of UnrealEngine containers on the SDK side
+ */
+
+
+/**
+ * SDK ext app type (reflect UE struct FGuid)
+ */
+struct FGuid
+{
+	uint32 A = 0;
+	uint32 B = 0;
+	uint32 C = 0;
+	uint32 D = 0;
+
+public:
+	inline bool operator==(const FGuid& In) const
+	{
+		return A == In.A && B == In.B && C == In.C && D == In.D;
+	}
+};
+
+/**
+ * SDK ext app type (reflect UE struct FIntPoint)
+ */
+struct FIntPoint
+{
+	int32 X = 0;
+	int32 Y = 0;
+
+public:
+	FIntPoint() = default;
+	FIntPoint(int32 InX, int32 InY)
+		: X(InX), Y(InY)
+	{ }
+
+public:
+	static FIntPoint ZeroValue;
+};
+
+/**
+ * SDK ext app type (reflect UE struct FIntRect)
+ */
+struct FIntRect
+{
+	FIntPoint Min;
+	FIntPoint Max;
+
+public:
+	FIntRect() = default;
+	FIntRect(const FIntPoint& InMin, const FIntPoint& InMax)
+		: Min(InMin), Max(InMax)
+	{ }
+};
+
+/**
+ * SDK ext app type (reflect UE struct FVector)
+ */
+struct FVector
+{
+	double X = 0;
+	double Y = 0;
+	double Z = 0;
+
+public:
+	FVector() = default;
+	FVector(double InX, double InY, double InZ)
+		: X(InX), Y(InY), Z(InZ)
+	{ }
+
+public:
+	static FVector ZeroVector;
+};
+
+/**
+ * SDK ext app type (reflect UE struct FQuat)
+ */
+struct FQuat
+{
+	double X = 0;
+	double Y = 0;
+	double Z = 0;
+	double W = 1;
+
+public:
+	FQuat() = default;
+	FQuat(double InX, double InY, double InZ, double InW)
+		: X(InX), Y(InY), Z(InZ), W(InW)
+	{ }
+};
+
+/**
+ * SDK ext app type (reflect UE struct FVector2D)
+ */
+struct FVector2D
+{
+	double X = 0;
+	double Y = 0;
+
+public:
+	FVector2D() = default;
+	FVector2D(double InX, double InY)
+		: X(InX), Y(InY)
+	{ }
+};
+
+/**
+ * SDK ext app type (reflect UE struct FMatrix)
+ */
+struct FMatrix
+{
+	double M[4][4] = { {1,0,0,0}, {0,1,0,0}, {0,0,1,0}, {0,0,0,1} };
+
+public:
+	FMatrix() = default;
+
+public:
+	static FMatrix Identity;
+};
+
+/**
+ * SDK ext app type (reflect UE struct FRotator)
+ */
+struct FRotator
+{
+	/** Rotation around the right axis (around Y axis), Looking up and down (0=Straight Ahead, +Up, -Down) */
+	double Pitch = 0;
+
+	/** Rotation around the up axis (around Z axis), Running in circles 0=East, +North, -South. */
+	double Yaw = 0;
+
+	/** Rotation around the forward axis (around X axis), Tilting your head, 0=Straight, +Clockwise, -CCW. */
+	double Roll = 0;
+
+public:
+	FRotator() = default;
+	FRotator(double InPitch, double InYaw, double InRoll)
+		: Pitch(InPitch), Yaw(InYaw), Roll(InRoll)
+	{ }
+
+public:
+	static FRotator ZeroRotator;
+};
+#endif
diff --git a/Engine/Extras/VirtualProduction/TextureShare/TextureShareSDK/Source/Containers/UnrealEngine/TextureShareSDKUnrealEngineEnums_PixelFormat.h b/Engine/Extras/VirtualProduction/TextureShare/TextureShareSDK/Source/Containers/UnrealEngine/TextureShareSDKUnrealEngineEnums_PixelFormat.h
new file mode 100644
--- /dev/null
+++ b/Engine/Extras/VirtualProduction/TextureShare/TextureShareSDK/Source/Containers/UnrealEngine/TextureShareSDKUnrealEngineEnums_PixelFormat.h
@@ -0,0 +1,99 @@
+// Copyright Epic Games, Inc. All Rights Reserved.
+
+#pragma once
+
+#ifndef __UNREAL__
+
+/**
+ * Reflect UE EPixelFormat enum on the SDK ext app side
+ * The code below is copied from the UE
+ */
+enum EPixelFormat : uint8
+{
+	PF_Unknown = 0,
+	PF_A32B32G32R32F = 1,
+	PF_B8G8R8A8 = 2,
+	PF_G8 = 3,
+	PF_G16 = 4,
+	PF_DXT1 = 5,
+	PF_DXT3 = 6,
+	PF_DXT5 = 7,
+	PF_UYVY = 8,
+	PF_FloatRGB = 9,
+	PF_FloatRGBA = 10,
+	PF_DepthStencil = 11,
+	PF_ShadowDepth = 12,
+	PF_R32_FLOAT = 13,
+	PF_G16R16 = 14,
+	PF_G16R16F = 15,
+	PF_G16R16F_FILTER = 16,
+	PF_G32R32F = 17,
+	PF_A2B10G10R10 = 18,
+	PF_A16B16G16R16 = 19,
+	PF_D24 = 20,
+	PF_R16F = 21,
+	PF_R16F_FILTER = 22,
+	PF_BC5 = 23,
+	PF_V8U8 = 24,
+	PF_A1 = 25,
+	PF_FloatR11G11B10 = 26,
+	PF_A8 = 27,
+	PF_R32_UINT = 28,
+	PF_R32_SINT = 29,
+	PF_PVRTC2 = 30,
+	PF_PVRTC4 = 31,
+	PF_R16_UINT = 32,
+	PF_R16_SINT = 33,
+	PF_R16G16B16A16_UINT = 34,
+	PF_R16G16B16A16_SINT = 35,
+	PF_R5G6B5_UNORM = 36,
+	PF_R8G8B8A8 = 37,
+	PF_A8R8G8B8 = 38,	// Only used for legacy loading; do NOT use!
+	PF_BC4 = 39,
+	PF_R8G8 = 40,
+	PF_ATC_RGB = 41,	// Unsupported Format
+	PF_ATC_RGBA_E = 42,	// Unsupported Format
+	PF_ATC_RGBA_I = 43,	// Unsupported Format
+	PF_X24_G8 = 44,	// Used for creating SRVs to alias a DepthStencil buffer to read Stencil. Don't use for creating textures.
+	PF_ETC1 = 45,	// Unsupported Format
+	PF_ETC2_RGB = 46,
+	PF_ETC2_RGBA = 47,
+	PF_R32G32B32A32_UINT = 48,
+	PF_R16G16_UINT = 49,
+	PF_ASTC_4x4 = 50,	// 8.00 bpp
+	PF_ASTC_6x6 = 51,	// 3.56 bpp
+	PF_ASTC_8x8 = 52,	// 2.00 bpp
+	PF_ASTC_10x10 = 53,	// 1.28 bpp
+	PF_ASTC_12x12 = 54,	// 0.89 bpp
+	PF_BC6H = 55,
+	PF_BC7 = 56,
+	PF_R8_UINT = 57,
+	PF_L8 = 58,
+	PF_XGXR8 = 59,
+	PF_R8G8B8A8_UINT = 60,
+	PF_R8G8B8A8_SNORM = 61,
+	PF_R16G16B16A16_UNORM = 62,
+	PF_R16G16B16A16_SNORM = 63,
+	PF_PLATFORM_HDR_0 = 64,
+	PF_PLATFORM_HDR_1 = 65,	// Reserved.
+	PF_PLATFORM_HDR_2 = 66,	// Reserved.
+	PF_NV12 = 67,
+	PF_R32G32_UINT = 68,
+	PF_ETC2_R11_EAC = 69,
+	PF_ETC2_RG11_EAC = 70,
+	PF_R8 = 71,
+	PF_B5G5R5A1_UNORM = 72,
+	PF_ASTC_4x4_HDR = 73,
+	PF_ASTC_6x6_HDR = 74,
+	PF_ASTC_8x8_HDR = 75,
+	PF_ASTC_10x10_HDR = 76,
+	PF_ASTC_12x12_HDR = 77,
+	PF_G16R16_SNORM = 78,
+	PF_R8G8_UINT = 79,
+	PF_R32G32B32_UINT = 80,
+	PF_R32G32B32_SINT = 81,
+	PF_R32G32B32F = 82,
+	PF_R8_SINT = 83,
+	PF_R64_UINT = 84
+};
+#endif
diff --git a/Engine/Extras/VirtualProduction/TextureShare/TextureShareSDK/Source/Containers/UnrealEngine/TextureShareSDKUnrealEngineMath.h b/Engine/Extras/VirtualProduction/TextureShare/TextureShareSDK/Source/Containers/UnrealEngine/TextureShareSDKUnrealEngineMath.h
new file mode 100644
--- /dev/null
+++ b/Engine/Extras/VirtualProduction/TextureShare/TextureShareSDK/Source/Containers/UnrealEngine/TextureShareSDKUnrealEngineMath.h
@@ -0,0 +1,54 @@
+// Copyright Epic Games, Inc. All Rights Reserved.
+
+#pragma once
+#include "Containers/UnrealEngine/TextureShareSDKUnrealEngineTypes.h"
+
+#ifndef __UNREAL__
+/**
+ * Implementation of UnrealEngine containers on the SDK side
+ */
+
+
+/**
+ * Reflect UE funcs to SDK ext app side
+ * The code below is copied from the UE
+ */
+struct FMath
+{
+	/** Computes absolute value in a generic way */
+	template< class T >
+	static constexpr inline T Abs(const T A)
+	{
+		return (A >= (T)0) ? A : -A;
+	}
+
+	/** Returns 1, 0, or -1 depending on relation of T to 0 */
+	template< class T >
+	static constexpr inline T Sign(const T A)
+	{
+		return (A > (T)0) ? (T)1 : ((A < (T)0) ? (T)-1 : (T)0);
+	}
+
+	/** Returns higher value in a generic way */
+	template< class T >
+	static constexpr inline T Max(const T A, const T B)
+	{
+		return (A >= B) ? A : B;
+	}
+
+	/** Returns lower value in a generic way */
+	template< class T >
+	static constexpr inline T Min(const T A, const T B)
+	{
+		return (A <= B) ? A : B;
+	}
+
+	// Allow mixing of float types to promote to highest precision type
+	static constexpr inline double Max(const double A, const float B) { return Max<double>(A, B); }
+	static constexpr inline double Max(const float A, const double B) { return Max<double>(A, B); }
+	static constexpr inline double Min(const double A, const float B) { return Min<double>(A, B); }
+	static constexpr inline double Min(const float A, const double B) { return Min<double>(A, B); }
+
+};
+
+#endif
diff --git a/Engine/Extras/VirtualProduction/TextureShare/TextureShareSDK/Source/Containers/UnrealEngine/TextureShareSDKUnrealEngineString.h b/Engine/Extras/VirtualProduction/TextureShare/TextureShareSDK/Source/Containers/UnrealEngine/TextureShareSDKUnrealEngineString.h
new file mode 100644
--- /dev/null
+++ b/Engine/Extras/VirtualProduction/TextureShare/TextureShareSDK/Source/Containers/UnrealEngine/TextureShareSDKUnrealEngineString.h
@@ -0,0 +1,291 @@
+// Copyright Epic Games, Inc. All Rights Reserved.
+
+#pragma once
+#include "Containers/UnrealEngine/TextureShareSDKUnrealEngineTypes.h"
+
+#ifndef __UNREAL__
+/**
+ * Implementation of UnrealEngine FString on the SDK side
+ * The code below is copied from the UE
+ */
+
+#include <xstring>
+#include <tchar.h>
+
+//@todo: implement better logic for std::wstring vs FString (this is too simplistic now)
+
+/**
+ * String support for SDK ext app side
+ */
+namespace ESearchCase
+{
+	enum Type
+	{
+		CaseSensitive,
+		IgnoreCase,
+	};
+}
+
+/**
+ * String support for SDK ext app side
+ */
+namespace FPlatformString
+{
+	/**
+	 * Type trait which yields an unsigned integer type of a given number of bytes.
+	 * If there is no such type, the Type member type will be absent, allowing it to be used in SFINAE contexts.
+	 */
+	template <int NumBytes>
+	struct TUnsignedIntType
+	{ };
+
+	template <> struct TUnsignedIntType<1> { using Type = uint8; };
+	template <> struct TUnsignedIntType<2> { using Type = uint16; };
+	template <> struct TUnsignedIntType<4> { using Type = uint32; };
+	template <> struct TUnsignedIntType<8> { using Type = uint64; };
+
+	/**
+	* Avoid sign extension problems with signed characters smaller than int
+	*
+	* E.g. 'Ö' - 'A' is negative since the char 'Ö' (0xD6) is negative and gets
+	* sign-extended to the 32-bit int 0xFFFFFFD6 before subtraction happens.
+	*
+	* Mainly needed for subtraction and addition.
+	*/
+	template <typename CharType>
+	struct TChar
+	{
+		static constexpr uint32 ToUnsigned(CharType Char)
+		{
+			return (typename TUnsignedIntType<sizeof(CharType)>::Type)Char;
+		}
+	};
+
+	typedef wchar_t WIDECHAR;
+	typedef size_t SIZE_T;
+
+	// Copied from \\Engine\Source\Runtime\Core\Private\GenericPlatform\GenericPlatformStricmp.cpp
+	static constexpr uint8 LowerAscii[128] = {
+		0x00, 0x01, 0x02, 0x03, 0x04, 0x05, 0x06, 0x07, 0x08, 0x09, 0x0A, 0x0B, 0x0C, 0x0D, 0x0E, 0x0F,
+		0x10, 0x11, 0x12, 0x13, 0x14, 0x15, 0x16, 0x17, 0x18, 0x19, 0x1A, 0x1B, 0x1C, 0x1D, 0x1E, 0x1F,
+		0x20, 0x21, 0x22, 0x23, 0x24, 0x25, 0x26, 0x27, 0x28, 0x29, 0x2A, 0x2B, 0x2C, 0x2D, 0x2E, 0x2F,
+		0x30, 0x31, 0x32, 0x33, 0x34, 0x35, 0x36, 0x37, 0x38, 0x39, 0x3A, 0x3B, 0x3C, 0x3D, 0x3E, 0x3F,
+		0x40, 'a',  'b',  'c',  'd',  'e',  'f',  'g',  'h',  'i',  'j',  'k',  'l',  'm',  'n',  'o',
+		'p',  'q',  'r',  's',  't',  'u',  'v',  'w',  'x',  'y',  'z',  0x5B, 0x5C, 0x5D, 0x5E, 0x5F,
+		0x60, 0x61, 0x62, 0x63, 0x64, 0x65, 0x66, 0x67, 0x68, 0x69, 0x6A, 0x6B, 0x6C, 0x6D, 0x6E, 0x6F,
+		0x70, 0x71, 0x72, 0x73, 0x74, 0x75, 0x76, 0x77, 0x78, 0x79, 0x7A, 0x7B, 0x7C, 0x7D, 0x7E, 0x7F
+	};
+
+	inline bool BothAscii(uint32 C1, uint32 C2)
+	{
+		return ((C1 | C2) & 0xffffff80) == 0;
+	}
+
+	inline int32 Strcmp(const WIDECHAR* String1, const WIDECHAR* String2)
+	{
+		return (int32)_tcscmp(String1, String2);
+	}
+
+	template<typename CharType1, typename CharType2>
+	inline int32 StricmpImpl(const CharType1* String1, const CharType2* String2)
+	{
+		while (true)
+		{
+			CharType1 C1 = *String1++;
+			CharType2 C2 = *String2++;
+
+			uint32 U1 = TChar<CharType1>::ToUnsigned(C1);
+			uint32 U2 = TChar<CharType2>::ToUnsigned(C2);
+
+			// Quickly move on if characters are identical but
+			// return equals if we found two null terminators
+			if (U1 == U2)
+			{
+				if (U1)
+				{
+					continue;
+				}
+
+				return 0;
+			}
+			else if (BothAscii(U1, U2))
+			{
+				if (int32 Diff = LowerAscii[U1] - LowerAscii[U2])
+				{
+					return Diff;
+				}
+			}
+			else
+			{
+				return U1 - U2;
+			}
+		}
+	}
+
+	template<typename CharType1, typename CharType2>
+	inline int32 StrnicmpImpl(const CharType1* String1, const CharType2* String2, SIZE_T Count)
+	{
+		for (; Count > 0; --Count)
+		{
+			CharType1 C1 = *String1++;
+			CharType2 C2 = *String2++;
+
+			// Quickly move on if characters are identical but
+			// return equals if we found two null terminators
+			if (C1 == C2)
+			{
+				if (C1)
+				{
+					continue;
+				}
+
+				return 0;
+			}
+			else if (BothAscii(C1, C2))
+			{
+				if (int32 Diff = LowerAscii[TChar<CharType1>::ToUnsigned(C1)] - LowerAscii[TChar<CharType2>::ToUnsigned(C2)])
+				{
+					return Diff;
+				}
+			}
+			else
+			{
+				return TChar<CharType1>::ToUnsigned(C1) - TChar<CharType2>::ToUnsigned(C2);
+			}
+		}
+
+		return 0;
+	}
+
+	inline int32 Stricmp(const WIDECHAR* Str1, const WIDECHAR* Str2) { return StricmpImpl(Str1, Str2); }
+};
+
+// Implement FString simple container for SDK data type mirroring
+struct FString
+{
+	std::wstring  Data;
+
+	typedef wchar_t CharType;
+	typedef int32   SizeType;
+
+public:
+	FString() = default;
+
+	FString(const FString& In)
+	{
+		Data = *In;
+	}
+
+	FString(const CharType* In)
+	{
+		Data = In;
+	}
+
+	FString(const CharType* In1, const CharType* In2)
+	{
+		Data = In1;
+		Data += In2;
+	}
+
+	FString(const CharType* In1, const CharType* In2, const CharType* In3)
+	{
+		Data = In1;
+		Data += In2;
+		Data += In3;
+	}
+
+	FString(const CharType* In1, const CharType* In2, const CharType* In3, const CharType* In4)
+	{
+		Data = In1;
+		Data += In2;
+		Data += In3;
+		Data += In4;
+	}
+
+public:
+	inline void operator=(const FString& In)
+	{
+		Data = *In;
+	}
+
+	inline void operator=(const CharType* In)
+	{
+		Data = In;
+	}
+
+	inline const CharType* operator*() const
+	{
+		return Data.c_str();
+	}
+
+	inline SizeType Len() const
+	{
+		return (SizeType)Data.length();
+	}
+
+	inline bool IsEmpty() const
+	{
+		return Data.empty();
+	}
+
+public:
+	/**
+	 * Lexicographically tests whether this string is equivalent to the Other given string
+	 *
+	 * @param Other 	The string test against
+	 * @param SearchCase 	Whether or not the comparison should ignore case
+	 * @return true if this string is lexicographically equivalent to the other, otherwise false
+	 */
+	inline bool Equals(const FString& Other, ESearchCase::Type SearchCase = ESearchCase::CaseSensitive) const
+	{
+		SizeType Num = Len();
+		SizeType OtherNum = Other.Len();
+
+		if (Num != OtherNum)
+		{
+			// Handle special case where FString() == FString("")
+			return Num + OtherNum == 1;
+		}
+		else if (Num > 1)
+		{
+			if (SearchCase == ESearchCase::CaseSensitive)
+			{
+				return FPlatformString::Strcmp(Data.c_str(), Other.Data.c_str()) == 0;
+			}
+			else
+			{
+				return FPlatformString::Stricmp(Data.c_str(), Other.Data.c_str()) == 0;
+			}
+		}
+
+		return true;
+	}
+
+	/**
+	 * Lexicographically test whether the left string is == the right string
+	 *
+	 * @param Lhs String to compare against.
+	 * @param Rhs String to compare against.
+	 * @return true if the left string is lexicographically == the right string, otherwise false
+	 * @note case insensitive
+	 */
+	friend inline bool operator==(const FString& Lhs, const FString& Rhs)
+	{
+		return Lhs.Equals(Rhs, ESearchCase::IgnoreCase);
+	}
+
+	/**
+	 * Lexicographically test whether the left string is == the right string
+	 *
+	 * @param Lhs String to compare against.
+	 * @param Rhs String to compare against.
+	 * @return true if the left string is lexicographically == the right string, otherwise false
+	 * @note case insensitive
+	 */
+	friend inline bool operator==(const FString& Lhs, const CharType* Rhs)
+	{
+		return FPlatformString::Stricmp(*Lhs, Rhs) == 0;
+	}
+};
+
+#endif
diff --git a/Engine/Extras/VirtualProduction/TextureShare/TextureShareSDK/Source/Containers/UnrealEngine/TextureShareSDKUnrealEngineTypes.h b/Engine/Extras/VirtualProduction/TextureShare/TextureShareSDK/Source/Containers/UnrealEngine/TextureShareSDKUnrealEngineTypes.h
new file mode 100644
--- /dev/null
+++ b/Engine/Extras/VirtualProduction/TextureShare/TextureShareSDK/Source/Containers/UnrealEngine/TextureShareSDKUnrealEngineTypes.h
@@ -0,0 +1,36 @@
+// Copyright Epic Games, Inc. All Rights Reserved.
+
+#pragma once
+
+#ifndef __UNREAL__
+/**
+ * Implementation of atomic UE types on the SDK side
+ * The code below is copied from the UE
+ */
+
+namespace Windows
+{
+	typedef void* HANDLE;
+};
+
+// Standard typedefs
+typedef signed char         schar;
+typedef signed char         int8;
+typedef short               int16;
+typedef int                 int32;
+#ifdef _MSC_VER
+typedef __int64             int64;
+#else
+typedef long long           int64;
+#endif /* _MSC_VER */
+
+typedef unsigned char      uint8;
+typedef unsigned short     uint16;
+typedef unsigned int       uint32;
+#ifdef _MSC_VER
+typedef unsigned __int64   uint64;
+#else
+typedef unsigned long long uint64;
+#endif /* _MSC_VER */
+
+#endif
diff --git a/Engine/Extras/VirtualProduction/TextureShare/TextureShareSDK/Source/ITextureShareSDK.h b/Engine/Extras/VirtualProduction/TextureShare/TextureShareSDK/Source/ITextureShareSDK.h
new file mode 100644
--- /dev/null
+++ b/Engine/Extras/VirtualProduction/TextureShare/TextureShareSDK/Source/ITextureShareSDK.h
@@ -0,0 +1,74 @@
+// Copyright Epic Games, Inc. All Rights Reserved.
+
+#pragma once
+#include "ITextureShareSDKObject.h"
+
+/**
+ * TextureShareSDK API
+ */
+class ITextureShareSDK
+{
+public:
+	/**
+	* Get exist or create new ITextureShareSDKObject object
+	*
+	* @param ShareName   - Unique share name
+	*
+	* @return - Pointer to an existing/new interface, or null (if the name is already in use in another interface type)
+	*/
+	TEXTURESHARESDK_API static ITextureShareSDKObject* GetOrCreateObject(const wchar_t* ShareName);
+
+	/**
+	 * Check if ITextureShareSDKObject object with that name exist
+	 *
+	 * @param ShareName   - Unique share name
+	 *
+	 * @return - true if exist
+	 */
+	TEXTURESHARESDK_API static  bool IsObjectExist(const wchar_t* ShareName);
+
+	/**
+	 * Release ITextureShareSDKObject object
+	 *
+	 * @param ShareName   - Unique share name
+	 *
+	 * @return - true if released
+	 */
+	TEXTURESHARESDK_API static  bool RemoveObject(const wchar_t* ShareName);
+
+	/**
+	 * Returns a description of all objects
+	 *
+	 * @param OutInterprocessObjects - All exist texture share objects on this computer
+	 * @param ShareName              - (optional) return only objects with this name
+	 *
+	 * @return - true if success
+	 */
+	TEXTURESHARESDK_API static bool GetInterprocessObjects(TDataOutput<TArraySerializable<FTextureShareCoreObjectDesc>>& OutInterprocessObjects, const wchar_t* ShareName = nullptr);
+
+	/**
+	 * Returns detailed information about the local process
+	 */
+	TEXTURESHARESDK_API static void GetProcessDesc(TDataOutput<FTextureShareCoreObjectProcessDesc>& OutProcessDesc);
+
+	/**
+	 * Assign a unique name to the local process. The changes only affect newly created TextureShare objects.
+	 *
+	 * @param InProcessId - Custom local process name
+	 */
+	TEXTURESHARESDK_API static void SetProcessName(const wchar_t* InProcessId);
+
+	/**
+	 * Assign a device type to the local process. The changes only affect newly created TextureShare objects.
+	 *
+	 * @param InDeviceType - render device type
+	 */
+	TEXTURESHARESDK_API static bool SetProcessDeviceType(const ETextureShareDeviceType InDeviceType);
+
+	/**
+	 * Release all resources last used more than specified timeout value
+	 *
+	 * @param InMilisecondsTimeOut - Timeout to remove resource
+	 */
+	TEXTURESHARESDK_API static void RemoveUnusedResources(const uint32 InMilisecondsTimeOut);
+};
diff --git a/Engine/Extras/VirtualProduction/TextureShare/TextureShareSDK/Source/ITextureShareSDKObject.h b/Engine/Extras/VirtualProduction/TextureShare/TextureShareSDK/Source/ITextureShareSDKObject.h
new file mode 100644
--- /dev/null
+++ b/Engine/Extras/VirtualProduction/TextureShare/TextureShareSDK/Source/ITextureShareSDKObject.h
@@ -0,0 +1,271 @@
+// Copyright Epic Games, Inc. All Rights Reserved.
+
+#pragma once
+#include "Containers/TextureShareSDKContainers.h"
+#include "Containers/TextureShareCoreContainers.h"
+#include "Serialize/TextureShareSDKDataWrappers.h"
+
+/**
+ * TextureShareSDK object API
+ */
+class TEXTURESHARESDK_API ITextureShareSDKObject
+{
+public:
+	virtual ~ITextureShareSDKObject() = default;
+
+public:
+	///////////////////////// State /////////////////////////
+
+	/**
+	* Return the name of the TextureShare.
+	*/
+	virtual const wchar_t* GetName() const = 0;
+
+	/**
+	* Returns detailed information about the TextureShare object.
+	*/
+	virtual void GetObjectDesc(TDataOutput<FTextureShareCoreObjectDesc>& OutObjectDesc) const = 0;
+
+	/**
+	 * Returns true if the TextureShare object is ready to be used.
+	 */
+	virtual bool IsActive() const = 0;
+
+	/**
+	 * Returns true if the TextureShareCore object has started a session and processes are connected for this frame.
+	 */
+	virtual bool IsFrameSyncActive() const = 0;
+
+public:
+	///////////////////////// Settings /////////////////////////
+
+	/**
+	 * Change the process name for this TextureShare object.
+	 * If the object is synced now, then the changes are deferred and executed in the next frame.
+	 *
+	 * @param InProcessId - New process name
+	 *
+	 * @return true, if success
+	 */
+	virtual bool SetProcessId(const wchar_t* InProcessId) = 0;
+
+	/**
+	 * Assign a device type for this TextureShare object.
+	 *
+	 * @param InDeviceType - render device type
+	 */
+	virtual bool SetDeviceType(const ETextureShareDeviceType InDeviceType) = 0;
+
+	/**
+	 * Change the sync settings for this TextureShare object.
+	 * If the object is synced now, then the changes are deferred and executed in the next frame.
+	 *
+	 * @param InSyncSetting - New sync settings
+	 *
+	 * @return true, if success
+	 */
+	virtual bool SetSyncSetting(const TDataInput<FTextureShareCoreSyncSettings>& InSyncSetting) = 0;
+
+	/**
+	* Return the sync settings of the TextureShare.
+	*/
+	virtual void GetSyncSetting(TDataOutput<FTextureShareCoreSyncSettings>& OutSyncSettings) const = 0;
+
+	/**
+	 * Get TextureShare default sync settings by template type
+	 * The settings of this template are not related to the current settings of the object
+	 *
+	 * @param InType - Sync settings template type
+	 *
+	 * @return structure with sync logic settings
+	 */
+	virtual void GetFrameSyncSettings(const ETextureShareFrameSyncTemplate InType, TDataOutput<FTextureShareCoreFrameSyncSettings>& OutFrameSyncSettings) const = 0;
+
+public:
+	///////////////////////// Session /////////////////////////
+
+	/**
+	 * Start a TextureShare session for the specified device
+	 *
+	 * @param InDeviceType     - Shared resources render device type
+	 *
+	 * @return True if the success
+	 */
+	virtual bool BeginSession() = 0;
+
+	/**
+	 * End the TextureShare session and the resources in use are also removed.
+	 *
+	 * @return True if the success
+	 */
+	virtual bool EndSession() = 0;
+
+	/**
+	 * Returns true if the session is currently valid
+	 */
+	virtual bool IsSessionActive() const = 0;
+
+public:
+	///////////////////////// Thread sync support /////////////////////////
+
+	/**
+	 * Support for multi-threaded implementation.
+	 * The mutex of the specified type will be locked (and created for the first time)
+	 *
+	 * @param InThreadMutex - Mutex type
+	 *
+	 * @return True if the success
+	 */
+	virtual bool LockThreadMutex(const ETextureShareThreadMutex InThreadMutex) = 0;
+
+	/**
+	 * Support for multi-threaded implementation.
+	 * The mutex of the specified type will be unlocked
+	 *
+	 * @param InThreadMutex - Mutex type
+	 *
+	 * @return True if the success
+	 */
+	virtual bool UnlockThreadMutex(const ETextureShareThreadMutex InThreadMutex) = 0;
+
+	/**
+	 * Support for multi-threaded implementation.
+	 * Returns true if the BeginFrameSync() function can be called at this moment
+	 */
+	virtual bool IsBeginFrameSyncActive() const = 0;
+
+	/**
+	 * Support for multi-threaded implementation.
+	 * Returns true if the BeginFrameSync_RenderThread() function can be called at this moment
+	 */
+	virtual bool IsBeginFrameSyncActive_RenderThread() const = 0;
+
+public:
+	///////////////////////// Interprocess Synchronization /////////////////////////
+
+	/**
+	 * Begin sync logic in range FrameBegin..FrameEnd
+	 * The list of connected processes for the current frame will also be updated.
+	 * Returns true if the frame is connected to other processes that match the synchronization settings.
+	 */
+	virtual bool BeginFrameSync() = 0;
+
+	/**
+	 * Synchronize connected processes that support this sync step
+	 * Data from remote processes will be read at the time the barrier is synchronized.
+	 * Missed sync steps from the sync settings will also be called.
+	 *
+	 * @param InSyncStep - Sync step value
+	 *
+	 * @return True if the success
+	 */
+	virtual bool FrameSync(const ETextureShareSyncStep InSyncStep) = 0;
+
+	/**
+	 * Finalize sync logic in range FrameBegin..FrameEnd
+	 * Missed sync steps from the sync settings will also be called.
+	 */
+	virtual bool EndFrameSync() = 0;
+
+	/**
+	 * Begin sync logic in range FrameProxyBegin..FrameProxyEnd
+	 */
+	virtual bool BeginFrameSync_RenderThread() = 0;
+
+	/**
+	 * Synchronize connected processes that support this sync step
+	 * ProxyData from remote processes will be read at the time the barrier is synchronized.
+	 * Missed sync steps from the sync settings will also be called.
+	 *
+	 * @param InSyncStep - Sync step value
+	 *
+	 * @return True if the success
+	 */
+	virtual bool FrameSync_RenderThread(const ETextureShareSyncStep InSyncStep) = 0;
+
+	/**
+	 * Finalize sync logic in range FrameProxyBegin..FrameProxyEnd
+	 * Missed sync steps from the sync settings will also be called.
+	 */
+	virtual bool EndFrameSync_RenderThread() = 0;
+
+	/**
+	 * Returns a list of handles to the interprocess TextureShare objects that are currently connected to this object.
+	 *
+	 * @param OutInterprocessObjects - Output array with handles to TextureShare objects
+	 *
+	 * @return true, if success
+	 */
+	virtual void GetConnectedInterprocessObjects(TDataOutput<TArraySerializable<FTextureShareCoreObjectDesc>>& OutObjectList) const = 0;
+
+public:
+	///////////////////////// Data Containers /////////////////////////
+
+	/**
+	 * Get object data
+	 * Object data for the current frame in the GameThread.
+	 */
+	virtual void GetData(TDataOutput<FTextureShareCoreData>& OutObjectData) = 0;
+
+	/**
+	 * Set object data
+	 * Set new object data for the current frame in the GameThread.
+	 */
+	virtual void SetData(const TDataInput<FTextureShareCoreData>& InObjectData) = 0;
+
+	/**
+	 * Get object ProxyData.
+	 * Object proxy data for the current frame in the RenderThread.
+	 */
+	virtual void GetProxyData_RenderThread(TDataOutput<FTextureShareCoreProxyData>& OutObjectProxyData) = 0;
+
+	/**
+	 * Set object ProxyData.
+	 * Object proxy data for the current frame in the RenderThread.
+	 */
+	virtual void SetProxyData_RenderThread(const TDataInput<FTextureShareCoreProxyData>& InObjectProxyData) = 0;
+
+	/**
+	 * Received Data from connected process objects
+	 */
+	virtual void GetReceivedData(TDataOutput<TArraySerializable<FTextureShareCoreObjectData>>& OutReceivedObjectsData) const = 0;
+
+	/**
+	 * Received ProxyData from connected process objects
+	 */
+	virtual void GetReceivedProxyData_RenderThread(TDataOutput<TArraySerializable<FTextureShareCoreObjectProxyData>>& OutObjectsProxyData) const = 0;
+
+public:
+	///////////////////////// Access to shared resources /////////////////////////
+
+	/**
+	 * Open shared D3D11 resource
+	 *
+	 * @param InD3D11Device - D3D11 device interface
+	 * @param InResourceDesc- Shared resource descriptor
+	 *
+	 * @return pointer to shared D3D11 texture, or nullptr on failure
+	 */
+	virtual ID3D11Texture2D* OpenSharedResourceD3D11(ID3D11Device* InD3D11Device, const TDataInput<FTextureShareCoreResourceDesc>& InResourceDesc) = 0;
+
+	/**
+	 * Open shared D3D12 resource
+	 *
+	 * @param InD3D12Device - D3D12 device interface
+	 * @param InResourceDesc- Shared resource descriptor
+	 *
+	 * @return pointer to shared D3D12 texture, or nullptr on failure
+	 */
+	virtual  ID3D12Resource* OpenSharedResourceD3D12(ID3D12Device* InD3D12Device, const TDataInput<FTextureShareCoreResourceDesc>& InResourceDesc) = 0;
+
+	/**
+	 * Open shared Vulkan resource
+	 *
+	 * @parama InDeviceVulkanContext - Vulkan device context
+	 * @param InResourceDesc         - Shared resource descriptor
+	 * @param OutVulkanResource      - Output Vulkan resource
+	 *
+	 * @return true, if resource opened
+	 */
+	virtual  bool OpenSharedResourceVulkan(const TDataInput<FTextureShareDeviceVulkanContext>& InDeviceVulkanContext, const TDataInput<FTextureShareCoreResourceDesc>& InResourceDesc, TDataOutput<FTextureShareDeviceVulkanResource>& OutVulkanResource) = 0;
+};
diff --git a/Engine/Extras/VirtualProduction/TextureShare/TextureShareSDK/Source/Misc/TextureShareCoreStrings.h b/Engine/Extras/VirtualProduction/TextureShare/TextureShareSDK/Source/Misc/TextureShareCoreStrings.h
new file mode 100644
--- /dev/null
+++ b/Engine/Extras/VirtualProduction/TextureShare/TextureShareSDK/Source/Misc/TextureShareCoreStrings.h
@@ -0,0 +1,32 @@
+// Copyright Epic Games, Inc. All Rights Reserved.
+
+#pragma once
+#if __UNREAL__
+#include "CoreMinimal.h"
+#else
+#include "windows.h"
+#endif
+
+/**
+ * Default names for TextureShareCore module
+ */
+namespace TextureShareCoreStrings
+{
+	namespace Default
+	{
+		static constexpr auto ShareName = TEXT("DefaultShareName");
+
+		namespace ProcessName
+		{
+#if defined TEXTURESHARE_PROJECT_NAME
+			// Use TextureShare project name as process name
+			static constexpr auto SDK = TEXT(TEXTURESHARE_PROJECT_NAME);
+#else
+			static constexpr auto SDK = TEXT("TextureShareSDK");
+#endif
+			static constexpr auto UE = TEXT("UnrealEngine");
+		}
+
+		static constexpr auto ViewId = TEXT("DefaultView");
+	}
+};
diff --git a/Engine/Extras/VirtualProduction/TextureShare/TextureShareSDK/Source/Misc/TextureShareDisplayClusterStrings.h b/Engine/Extras/VirtualProduction/TextureShare/TextureShareSDK/Source/Misc/TextureShareDisplayClusterStrings.h
new file mode 100644
--- /dev/null
+++ b/Engine/Extras/VirtualProduction/TextureShare/TextureShareSDK/Source/Misc/TextureShareDisplayClusterStrings.h
@@ -0,0 +1,40 @@
+// Copyright Epic Games, Inc. All Rights Reserved.
+
+#pragma once
+#if __UNREAL__
+#include "CoreMinimal.h"
+#else
+#include "windows.h"
+#endif
+
+/**
+ * Constant names for nDisplay texture resources
+ */
+namespace TextureShareDisplayClusterStrings
+{
+	namespace Viewport
+	{
+		// FinalColor, but can be overrided
+		static constexpr auto FinalColor = TEXT("ViewportFinalColor");
+
+		// internal resolved shader resources, used as warp&blend source
+		static constexpr auto Input = TEXT("ViewportInput");
+		static constexpr auto Mips = TEXT("ViewportMips");
+
+		// After warp viewport (before output remap)
+		static constexpr auto Warped = TEXT("ViewportWarpBlend");
+	}
+
+	namespace Output
+	{
+		// access to nDisplay frame backbuffer
+		static constexpr auto Backbuffer     = TEXT("FrameBackbuffer");
+		static constexpr auto BackbufferTemp = TEXT("FrameBackbufferTemp");
+	}
+
+	namespace Default
+	{
+		// Default ShareName for nDisplay integration object
+		static constexpr auto ShareName = TEXT("nDisplay");
+	}
+};
diff --git a/Engine/Extras/VirtualProduction/TextureShare/TextureShareSDK/Source/Misc/TextureShareStrings.h b/Engine/Extras/VirtualProduction/TextureShare/TextureShareSDK/Source/Misc/TextureShareStrings.h
new file mode 100644
--- /dev/null
+++ b/Engine/Extras/VirtualProduction/TextureShare/TextureShareSDK/Source/Misc/TextureShareStrings.h
@@ -0,0 +1,32 @@
+// Copyright Epic Games, Inc. All Rights Reserved.
+
+#pragma once
+#if __UNREAL__
+#include "CoreMinimal.h"
+#else
+#include "windows.h"
+#endif
+
+/**
+ * Default names for TextureShare module
+ */
+namespace TextureShareStrings
+{
+	namespace SceneTextures
+	{
+		// Read-only scene resources
+		static constexpr auto SceneColor = TEXT("SceneColor");
+		static constexpr auto SceneDepth = TEXT("SceneDepth");
+		static constexpr auto SmallDepthZ = TEXT("SmallDepthZ");
+		static constexpr auto GBufferA = TEXT("GBufferA");
+		static constexpr auto GBufferB = TEXT("GBufferB");
+		static constexpr auto GBufferC = TEXT("GBufferC");
+		static constexpr auto GBufferD = TEXT("GBufferD");
+		static constexpr auto GBufferE = TEXT("GBufferE");
+		static constexpr auto GBufferF = TEXT("GBufferF");
+
+		// Read-write RTTs
+		static constexpr auto FinalColor = TEXT("FinalColor");
+		static constexpr auto Backbuffer = TEXT("Backbuffer");
+	}
+};
diff --git a/Engine/Extras/VirtualProduction/TextureShare/TextureShareSDK/Source/Serialize/ITextureShareSDKSerializeStream.h b/Engine/Extras/VirtualProduction/TextureShare/TextureShareSDK/Source/Serialize/ITextureShareSDKSerializeStream.h
new file mode 100644
--- /dev/null
+++ b/Engine/Extras/VirtualProduction/TextureShare/TextureShareSDK/Source/Serialize/ITextureShareSDKSerializeStream.h
@@ -0,0 +1,139 @@
+// Copyright Epic Games, Inc. All Rights Reserved.
+
+#pragma once
+#include "Containers/TextureShareSDKContainers.h"
+#include "Serialize/ITextureShareSerialize.h"
+
+#ifndef __UNREAL__
+
+/**
+ * Support SDK data types serialization
+ */
+class ITextureShareSDKSerializeStream
+	: public ITextureShareSerializeStream
+{
+public:
+	virtual ~ITextureShareSDKSerializeStream() = default;
+
+	// Serialize any serializable objects
+	virtual ITextureShareSerializeStream& operator<<(ITextureShareSerialize& InOutData) override
+	{
+		return InOutData.Serialize(*this);
+	}
+
+public:
+	// Serialize atomic types:
+	virtual ITextureShareSerializeStream& operator<<(int8& In)  override { return this->SerializeData(&In, sizeof(In)); }
+	virtual ITextureShareSerializeStream& operator<<(int16& In) override { return this->SerializeData(&In, sizeof(In)); }
+	virtual ITextureShareSerializeStream& operator<<(int32& In) override { return this->SerializeData(&In, sizeof(In)); }
+	virtual ITextureShareSerializeStream& operator<<(int64& In) override { return this->SerializeData(&In, sizeof(In)); }
+
+	virtual ITextureShareSerializeStream& operator<<(uint8& In)  override { return this->SerializeData(&In, sizeof(In)); }
+	virtual ITextureShareSerializeStream& operator<<(uint16& In) override { return this->SerializeData(&In, sizeof(In)); }
+	virtual ITextureShareSerializeStream& operator<<(uint32& In) override { return this->SerializeData(&In, sizeof(In)); }
+	virtual ITextureShareSerializeStream& operator<<(uint64& In) override { return this->SerializeData(&In, sizeof(In)); }
+
+	virtual ITextureShareSerializeStream& operator<<(float& In)  override { return this->SerializeData(&In, sizeof(In)); }
+	virtual ITextureShareSerializeStream& operator<<(double& In) override { return this->SerializeData(&In, sizeof(In)); }
+	virtual ITextureShareSerializeStream& operator<<(bool& In)   override { return this->SerializeData(&In, sizeof(In)); }
+
+public:
+	// Serialize TextureShare enums:
+	virtual ITextureShareSerializeStream& operator<<(ETextureShareFrameSyncTemplate& In) override { return this->SerializeData(&In, sizeof(In)); }
+
+	virtual ITextureShareSerializeStream& operator<<(ETextureShareSyncPass& In)     override { return this->SerializeData(&In, sizeof(In)); }
+	virtual ITextureShareSerializeStream& operator<<(ETextureShareSyncState& In)    override { return this->SerializeData(&In, sizeof(In)); }
+	virtual ITextureShareSerializeStream& operator<<(ETextureShareSyncStep& In)    override { return this->SerializeData(&In, sizeof(In)); }
+	virtual ITextureShareSerializeStream& operator<<(ETextureShareTextureOp& In)   override { return this->SerializeData(&In, sizeof(In)); }
+	virtual ITextureShareSerializeStream& operator<<(ETextureShareEyeType& In)     override { return this->SerializeData(&In, sizeof(In)); }
+	virtual ITextureShareSerializeStream& operator<<(ETextureShareDeviceType& In)  override { return this->SerializeData(&In, sizeof(In)); }
+	virtual ITextureShareSerializeStream& operator<<(ETextureShareProcessType& In) override { return this->SerializeData(&In, sizeof(In)); }
+
+	virtual ITextureShareSerializeStream& operator<<(ETextureShareCoreSceneViewManualProjectionType& In)   override { return this->SerializeData(&In, sizeof(In)); }
+	virtual ITextureShareSerializeStream& operator<<(ETextureShareViewRotationDataType& In) override { return this->SerializeData(&In, sizeof(In)); }
+	virtual ITextureShareSerializeStream& operator<<(ETextureShareViewLocationDataType& In) override { return this->SerializeData(&In, sizeof(In)); }
+
+public:
+	// Serialize other enums
+	virtual ITextureShareSerializeStream& operator<<(EPixelFormat& In) override { return this->SerializeData(&In, sizeof(In)); }
+	virtual ITextureShareSerializeStream& operator<<(DXGI_FORMAT& In)  override { return this->SerializeData(&In, sizeof(In)); }
+
+public:
+	// Serialize UECore types:
+	virtual ITextureShareSerializeStream& operator<<(FString& InOutData) override
+	{
+		int32 StrLen = InOutData.Len();
+		*this << StrLen;
+
+		if (StrLen > 0)
+		{
+			if (IsWriteStream())
+			{
+				const wchar_t* WCharValue = *InOutData;
+				SerializeData(WCharValue, sizeof(wchar_t) * StrLen);
+			}
+			else
+			{
+				TArray<wchar_t> WCharBuffer;
+				WCharBuffer.AddZeroed(StrLen + 1);
+				SerializeData(WCharBuffer.GetData(), sizeof(wchar_t) * StrLen);
+				const wchar_t* WCharValue = WCharBuffer.GetData();
+
+				InOutData = WCharValue;
+			}
+		}
+
+		return *this;
+	}
+
+	virtual ITextureShareSerializeStream& operator<<(FVector& InOutData) override
+	{
+		return (*this) << InOutData.X << InOutData.Y << InOutData.Z;
+	}
+
+	virtual ITextureShareSerializeStream& operator<<(FVector2D& InOutData) override
+	{
+		return (*this) << InOutData.X << InOutData.Y;
+	}
+
+	virtual ITextureShareSerializeStream& operator<<(FRotator& InOutData) override
+	{
+		return (*this) << InOutData.Pitch << InOutData.Yaw << InOutData.Roll;
+	}
+
+	virtual ITextureShareSerializeStream& operator<<(FQuat& InOutData) override
+	{
+		return (*this) << InOutData.X << InOutData.Y << InOutData.Z << InOutData.W;
+	}
+
+	virtual ITextureShareSerializeStream& operator<<(FMatrix& InOutData) override
+	{
+		return (*this)
+			<< InOutData.M[0][0] << InOutData.M[0][1] << InOutData.M[0][2] << InOutData.M[0][3]
+			<< InOutData.M[1][0] << InOutData.M[1][1] << InOutData.M[1][2] << InOutData.M[1][3]
+			<< InOutData.M[2][0] << InOutData.M[2][1] << InOutData.M[2][2] << InOutData.M[2][3]
+			<< InOutData.M[3][0] << InOutData.M[3][1] << InOutData.M[3][2] << InOutData.M[3][3];
+	}
+
+	virtual ITextureShareSerializeStream& operator<<(FIntRect& InOutData) override
+	{
+		return (*this) << InOutData.Min << InOutData.Max;
+	}
+
+	virtual ITextureShareSerializeStream& operator<<(FIntPoint& InOutData) override
+	{
+		return (*this) << InOutData.X << InOutData.Y;
+	}
+
+	virtual ITextureShareSerializeStream& operator<<(FGuid& InOutData) override
+	{
+		return (*this) << InOutData.A << InOutData.B << InOutData.C << InOutData.D;
+	}
+
+	virtual ITextureShareSerializeStream& operator<<(Windows::HANDLE& InOutData) override
+	{
+		return SerializeData(&InOutData, sizeof(Windows::HANDLE));
+	}
+};
+
+#endif
diff --git a/Engine/Extras/VirtualProduction/TextureShare/TextureShareSDK/Source/Serialize/ITextureShareSerialize.h b/Engine/Extras/VirtualProduction/TextureShare/TextureShareSDK/Source/Serialize/ITextureShareSerialize.h
new file mode 100644
--- /dev/null
+++ b/Engine/Extras/VirtualProduction/TextureShare/TextureShareSDK/Source/Serialize/ITextureShareSerialize.h
@@ -0,0 +1,103 @@
+// Copyright Epic Games, Inc. All Rights Reserved.
+
+#pragma once
+#if __UNREAL__
+#include "CoreMinimal.h"
+#endif
+
+#include "Containers/TextureShareCoreEnums.h"
+
+#if __UNREAL__
+#include "PixelFormat.h"
+#else
+// Provides an enumeration of UE pixel formats for the SDK.
+#include "Containers/UnrealEngine/TextureShareSDKUnrealEngineEnums_PixelFormat.h"
+#endif
+
+#include <DXGIFormat.h>
+
+class ITextureShareSerializeStream;
+
+/**
+ * Serialize structures data to exchange between processes
+ */
+struct ITextureShareSerialize
+{
+	virtual ~ITextureShareSerialize() = default;
+	virtual ITextureShareSerializeStream& Serialize(ITextureShareSerializeStream& Stream) = 0;
+};
+
+/**
+ * Interface for binary serialization
+ */
+class ITextureShareSerializeStream
+{
+public:
+	virtual ~ITextureShareSerializeStream() = default;
+
+public:
+	// Serialize any serializable objects
+	virtual ITextureShareSerializeStream& operator<<(ITextureShareSerialize& InOutData) = 0;
+
+public:
+	// Serialize atomic types:
+	virtual ITextureShareSerializeStream& operator<<(int8& In) = 0;
+	virtual ITextureShareSerializeStream& operator<<(int16& In) = 0;
+	virtual ITextureShareSerializeStream& operator<<(int32& In) = 0;
+	virtual ITextureShareSerializeStream& operator<<(int64& In) = 0;
+
+	virtual ITextureShareSerializeStream& operator<<(uint8& In) = 0;
+	virtual ITextureShareSerializeStream& operator<<(uint16& In) = 0;
+	virtual ITextureShareSerializeStream& operator<<(uint32& In) = 0;
+	virtual ITextureShareSerializeStream& operator<<(uint64& In) = 0;
+
+	virtual ITextureShareSerializeStream& operator<<(float& In) = 0;
+	virtual ITextureShareSerializeStream& operator<<(double& In) = 0;
+	virtual ITextureShareSerializeStream& operator<<(bool& In) = 0;
+
+public:
+	// Serialize TextureShare enums:
+	virtual ITextureShareSerializeStream& operator<<(ETextureShareFrameSyncTemplate& In) = 0;
+	virtual ITextureShareSerializeStream& operator<<(ETextureShareSyncPass& In) = 0;
+	virtual ITextureShareSerializeStream& operator<<(ETextureShareSyncState& In) = 0;
+
+	virtual ITextureShareSerializeStream& operator<<(ETextureShareSyncStep& In) = 0;
+	virtual ITextureShareSerializeStream& operator<<(ETextureShareTextureOp& In) = 0;
+	virtual ITextureShareSerializeStream& operator<<(ETextureShareEyeType& In) = 0;
+	virtual ITextureShareSerializeStream& operator<<(ETextureShareProcessType& In) = 0;
+	virtual ITextureShareSerializeStream& operator<<(ETextureShareDeviceType& In) = 0;
+
+	virtual ITextureShareSerializeStream& operator<<(ETextureShareCoreSceneViewManualProjectionType& In) = 0;
+	virtual ITextureShareSerializeStream& operator<<(ETextureShareViewRotationDataType& In) = 0;
+	virtual ITextureShareSerializeStream& operator<<(ETextureShareViewLocationDataType& In) = 0;
+
+public:
+	// Serialize other enums
+	virtual ITextureShareSerializeStream& operator<<(EPixelFormat& In) = 0;
+	virtual ITextureShareSerializeStream& operator<<(DXGI_FORMAT& In) = 0;
+
+public:
+	// Serialize UECore types:
+	virtual ITextureShareSerializeStream& operator<<(FString& InOutData) = 0;
+	virtual ITextureShareSerializeStream& operator<<(FVector& InOutData) = 0;
+	virtual ITextureShareSerializeStream& operator<<(FVector2D& InOutData) = 0;
+	virtual ITextureShareSerializeStream& operator<<(FRotator& InOutData) = 0;
+	virtual ITextureShareSerializeStream& operator<<(FQuat& InOutData) = 0;
+	virtual ITextureShareSerializeStream& operator<<(FMatrix& InOutData) = 0;
+	virtual ITextureShareSerializeStream& operator<<(FIntRect& InOutData) = 0;
+	virtual ITextureShareSerializeStream& operator<<(FIntPoint& InOutData) = 0;
+	virtual ITextureShareSerializeStream& operator<<(FGuid& InOutData) = 0;
+	virtual ITextureShareSerializeStream& operator<<(Windows::HANDLE& InOutData) = 0;
+
+public:
+	virtual bool IsWriteStream() const = 0;
+	virtual ITextureShareSerializeStream& SerializeData(void*, const uint32_t)
+	{
+		return *this;
+	}
+
+	virtual ITextureShareSerializeStream& SerializeData(const void*, const uint32_t)
+	{
+		return *this;
+	}
+};
diff --git a/Engine/Extras/VirtualProduction/TextureShare/TextureShareSDK/Source/Serialize/TextureShareCoreSerialize.h b/Engine/Extras/VirtualProduction/TextureShare/TextureShareSDK/Source/Serialize/TextureShareCoreSerialize.h
new file mode 100644
--- /dev/null
+++ b/Engine/Extras/VirtualProduction/TextureShare/TextureShareSDK/Source/Serialize/TextureShareCoreSerialize.h
@@ -0,0 +1,219 @@
+// Copyright Epic Games, Inc. All Rights Reserved.
+
+#pragma once
+#include "Serialize/ITextureShareSerialize.h"
+
+/**
+ * TTextureShareCoreRemoveReference<type> will remove any references from a type.
+ */
+template <typename T> struct TTextureShareCoreRemoveReference { typedef T Type; };
+template <typename T> struct TTextureShareCoreRemoveReference<T& > { typedef T Type; };
+template <typename T> struct TTextureShareCoreRemoveReference<T&&> { typedef T Type; };
+
+/**
+ * Forward will cast a reference to an rvalue reference.
+ * This is UE's equivalent of std::forward.
+ */
+template <typename T>
+inline T&& TextureShareCoreForward(typename TTextureShareCoreRemoveReference<T>::Type& Obj)
+{
+	return (T&&)Obj;
+}
+
+template <typename T>
+inline T&& TextureShareCoreForward(typename TTextureShareCoreRemoveReference<T>::Type&& Obj)
+{
+	return (T&&)Obj;
+}
+
+/**
+ * Invokes a callable with a set of arguments.  Allows the following:
+ * Copied from file "Templates\Invoke.h" for SDK
+ */
+template <typename FuncType, typename... ArgTypes>
+inline auto TextureShareCoreInvoke(FuncType&& Func, ArgTypes&&... Args)
+	-> decltype(TextureShareCoreForward<FuncType>(Func)(TextureShareCoreForward<ArgTypes>(Args)...))
+{
+	return TextureShareCoreForward<FuncType>(Func)(TextureShareCoreForward<ArgTypes>(Args)...);
+}
+
+/**
+ * Serializable array template for TextureShare
+ */
+template<typename T>
+class TArraySerializable
+	: public ITextureShareSerialize
+	, public TArray<T>
+{
+public:
+	virtual ITextureShareSerializeStream& Serialize(ITextureShareSerializeStream& Stream) override
+	{
+		uint32 ItemsCount = this->Num();
+		Stream << ItemsCount;
+
+		if(!Stream.IsWriteStream())
+		{
+			this->AddDefaulted(ItemsCount);
+		}
+
+		for (T& It : *this)
+		{
+			Stream << It;
+		}
+
+		return Stream;
+	}
+
+public:
+	// Append array values ordered. Value must support compare ops
+	void AppendSorted(const TArray<T>& InArrayValues)
+	{
+		for (const T& ValueIt : InArrayValues)
+		{
+			if (this->Find(ValueIt) == INDEX_NONE)
+			{
+				for (int32 Index = this->Num() - 1; Index >= 0; Index--)
+				{
+					if (this->operator[](Index) < ValueIt)
+					{
+						this->Insert(ValueIt, Index + 1);
+						break;
+					}
+				}
+			}
+		}
+	}
+
+	// Returns true if two array values are equal
+	bool EqualsFunc(const TArray<T>& InArrayValues) const
+	{
+		if (this->Num() == InArrayValues.Num())
+		{
+			for (int32 Index = 0; Index < this->Num(); Index++)
+			{
+				if (this->operator[](Index) != InArrayValues[Index])
+				{
+					return false;
+				}
+			}
+
+			return true;
+		}
+
+		return false;
+	}
+
+	/**
+	 * Finds an element which matches a predicate functor.
+	 *
+	 * @param Pred The functor to apply to each element. true, or nullptr if none is found.
+	 * @see FilterByPredicate, ContainsByPredicate
+	 */
+	template <typename Predicate>
+	T* FindByPredicate(Predicate Pred)
+	{
+		for (int32 Index = 0; Index < this->Num(); Index++)
+		{
+			T& ValueIt = this->operator[](Index);
+			if (TextureShareCoreInvoke(Pred, ValueIt))
+			{
+				return &ValueIt;
+			}
+		}
+
+		return nullptr;
+	}
+
+	template <typename Predicate>
+	const T* FindByPredicate(Predicate Pred) const
+	{
+		for (int32 Index = 0; Index < this->Num(); Index++)
+		{
+			const T& ValueIt = this->operator[](Index);
+			if (TextureShareCoreInvoke(Pred, ValueIt))
+			{
+				return &ValueIt;
+			}
+		}
+
+		return nullptr;
+	}
+
+	/**
+	 * Finds an item by predicate.
+	 *
+	 * @param Pred The predicate to match.
+	 * @returns Index to the first matching element, or INDEX_NONE if none is found.
+	 */
+	template <typename Predicate>
+	int32 IndexOfByPredicate(Predicate Pred) const
+	{
+		for (int32 Index = 0; Index < this->Num(); Index++)
+		{
+			const T& ValueIt = this->operator[](Index);
+			if (TextureShareCoreInvoke(Pred, ValueIt))
+			{
+				return Index;
+			}
+		}
+
+		return INDEX_NONE;
+	}
+
+	template<typename SubT>
+	const T* FindByEqualsFunc(const SubT& InValue) const
+	{
+		return FindByPredicate([InValue](const T& DataIt) {
+			return DataIt.EqualsFunc(InValue);
+		});
+	}
+
+	template<typename SubT>
+	bool GetValuesByEqualsFunc(const SubT& InValue, TArray<T>& OutValues) const
+	{
+		for (int32 Index = 0; Index < this->Num(); Index++)
+		{
+			const T& ValueIt = this->operator[](Index);
+			if (ValueIt.EqualsFunc(InValue))
+			{
+				OutValues.Add(ValueIt);
+			}
+		}
+
+		return OutValues.Num() > 0;
+	}
+
+	bool EmplaceValue(const T& InValue, const bool bOverrideExist = false)
+	{
+		if (T* ExistValue = FindByPredicate([InValue](const T& DataIt) {
+			return DataIt.EqualsFunc(InValue);
+		}))
+		{
+			if (bOverrideExist)
+			{
+				*ExistValue = InValue;
+			}
+
+			return bOverrideExist;
+		}
+
+		this->Add(InValue);
+
+		return true;
+	}
+
+	bool AppendValues(const TArraySerializable<T>& In, const bool bOverrideExist = false)
+	{
+		bool bResult = true;
+
+		for (const T& ResourceIt : *this)
+		{
+			if (EmplaceValue(ResourceIt, bOverrideExist) == false)
+			{
+				bResult = false;
+			}
+		}
+
+		return bResult;
+	}
+};
diff --git a/Engine/Extras/VirtualProduction/TextureShare/TextureShareSDK/Source/Serialize/TextureShareSDKDataWrappers.h b/Engine/Extras/VirtualProduction/TextureShare/TextureShareSDK/Source/Serialize/TextureShareSDKDataWrappers.h
new file mode 100644
--- /dev/null
+++ b/Engine/Extras/VirtualProduction/TextureShare/TextureShareSDK/Source/Serialize/TextureShareSDKDataWrappers.h
@@ -0,0 +1,253 @@
+// Copyright Epic Games, Inc. All Rights Reserved.
+
+#pragma once
+
+#include "Containers/TextureShareSDKContainers.h"
+
+#if __UNREAL__
+// implement binary serialization on the Unreal Engine side
+#include "Core/Serialize/ITextureShareCoreSerializeStream.h"
+using IDataContainerSerializeStream = ITextureShareCoreSerializeStream;
+
+#else
+// implement binary serialization on the user side
+#include "Serialize/ITextureShareSDKSerializeStream.h"
+using IDataContainerSerializeStream = ITextureShareSDKSerializeStream;
+
+/**
+ * Copied from UE for SDK ext app side
+ */
+namespace FPlatformMemory
+{
+	typedef size_t SIZE_T;
+
+	inline void* Memcpy(void* Dest, const void* Src, SIZE_T Count)
+	{
+		return memcpy(Dest, Src, Count);
+	}
+};
+
+#define check(Cond)
+#endif
+
+namespace DataSerializer
+{
+	enum Limits
+	{
+		// Transaction memory block size
+		MaxStreamDataSize = (32 * 1024),
+	};
+
+	/**
+	 * Binary serializer (Write)
+	 */
+	struct FStreamWrite
+		: public IDataContainerSerializeStream
+	{
+		uint8* const DataPtr = nullptr;
+		const uint32_t DataSize = 0;
+
+		uint32_t CurrentPos = 0;
+
+	public:
+		FStreamWrite(uint8* const InDataPtr, const uint32_t InDataSize)
+			: DataPtr(InDataPtr), DataSize(InDataSize)
+		{ }
+
+		virtual ~FStreamWrite() = default;
+
+		virtual bool IsWriteStream() const
+		{
+			return true;
+		}
+
+	protected:
+		virtual FStreamWrite& SerializeData(void* InDataPtr, const uint32_t InDataSize) override
+		{
+			return ImplWriteData(InDataPtr, InDataSize);
+		}
+
+		virtual FStreamWrite& SerializeData(const void* InDataPtr, const uint32_t InDataSize) override
+		{
+			return ImplWriteData(InDataPtr, InDataSize);
+		}
+
+	private:
+		FStreamWrite& ImplWriteData(const void* InDataPtr, const uint32_t InDataSize)
+		{
+			const uint32_t NextPos = CurrentPos + InDataSize;
+
+			check(NextPos <= DataSize);
+
+			// Copy data to the shared memory
+			FPlatformMemory::Memcpy(DataPtr + CurrentPos, InDataPtr, InDataSize);
+
+			// udate current writing pos
+			CurrentPos = NextPos;
+
+			return *this;
+		}
+	};
+
+	/**
+	 * Binary serializer (Read)
+	 */
+	class FStreamRead
+		: public IDataContainerSerializeStream
+	{
+		const uint8*   DataPtr = nullptr;
+		const uint32_t DataSize = 0;
+
+		uint32_t CurrentPos = 0;
+
+	public:
+		FStreamRead(const uint8* InDataPtr, const uint32_t InDataSize)
+			: DataPtr(InDataPtr), DataSize(InDataSize)
+		{ }
+
+		virtual ~FStreamRead() = default;
+
+		virtual bool IsWriteStream() const
+		{
+			return false;
+		}
+
+	protected:
+		virtual FStreamRead& SerializeData(void* InDataPtr, const uint32_t InDataSize) override
+		{
+			return ImplReadData(InDataPtr, InDataSize);
+		}
+
+	private:
+		FStreamRead& ImplReadData(void* InDataPtr, const uint32_t InDataSize)
+		{
+			if (InDataPtr && DataSize)
+			{
+				// Implement read
+				const uint32_t NextPos = CurrentPos + InDataSize;
+
+				// check vs structure size diffs
+				check(NextPos <= DataSize);
+
+				FPlatformMemory::Memcpy(InDataPtr, DataPtr + CurrentPos, InDataSize);
+
+				// Update current reading pos
+				CurrentPos = NextPos;
+			}
+
+			return *this;
+		}
+	};
+};
+
+/**
+ * Data wrapper SDK.DLL->ExtApp
+ */
+template<typename T>
+struct TDataOutput
+{
+	// Store structures as binary. And exchange like this
+	// Because the DLL versus the user application have differences in compilation, library version, memory alignment, etc.
+	uint32_t DataSize = 0;
+	uint8 DataPtr[DataSerializer::MaxStreamDataSize];
+public:
+	TDataOutput(T& InOutData)
+		: OutDataRef(InOutData)
+	{ }
+
+	~TDataOutput()
+	{
+		// Deserialize data
+		DataSerializer::FStreamRead SerializedDataArchive(DataPtr, DataSize);
+		SerializedDataArchive << OutDataRef;
+	}
+
+	TDataOutput<T>& operator=(const T& InData)
+	{
+		T Data(InData);
+		DataSerializer::FStreamWrite SerializedDataArchive(DataPtr, DataSerializer::MaxStreamDataSize);
+		SerializedDataArchive << Data;
+		DataSize = SerializedDataArchive.CurrentPos;
+
+		return *this;
+	}
+
+	TDataOutput<T>& operator=(T& InData)
+	{
+		DataSerializer::FStreamWrite SerializedDataArchive(DataPtr, DataSerializer::MaxStreamDataSize);
+		SerializedDataArchive << InData;
+		DataSize = SerializedDataArchive.CurrentPos;
+
+		return *this;
+	}
+
+	TDataOutput<T>& operator*()
+	{
+		return *this;
+	}
+
+private:
+	T& OutDataRef;
+};
+
+/**
+ * Data wrapper ExtApp->SDK.DLL
+ */
+template<typename T>
+struct TDataInput
+{
+	// Store structures as binary. And exchange like this
+	// Because the DLL versus the user application have differences in compilation, library version, memory alignment, etc.
+	uint8 DataPtr[DataSerializer::MaxStreamDataSize];
+	uint32_t DataSize = 0;
+
+public:
+	TDataInput(const T& InData)
+	{
+		T Data(InData);
+		DataSerializer::FStreamWrite SerializedDataArchive(DataPtr, DataSerializer::MaxStreamDataSize);
+		SerializedDataArchive << Data;
+		DataSize = SerializedDataArchive.CurrentPos;
+	}
+
+	TDataInput(T& InData)
+	{
+		DataSerializer::FStreamWrite SerializedDataArchive(DataPtr, DataSerializer::MaxStreamDataSize);
+		SerializedDataArchive << InData;
+		DataSize = SerializedDataArchive.CurrentPos;
+	}
+
+	inline T& Deserialize(T& OutData) const
+	{
+		// Deserialize data
+		DataSerializer::FStreamRead SerializedDataArchive(DataPtr, DataSize);
+		SerializedDataArchive << OutData;
+
+		return OutData;
+	}
+};
+
+/**
+ * Data deserializer TDataInput<InType> to InType
+ */
+template<typename T>
+struct TDataVariable
+{
+	T OutData;
+
+	TDataVariable(const TDataInput<T>& InData)
+	{
+		InData.Deserialize(OutData);
+	}
+
+
+	T& operator*()
+	{
+		return OutData;
+	}
+
+	const T& operator*() const
+	{
+		return OutData;
+	}
+};
