Added localization for Connection Type Settings

diff --git a/Engine/Extras/VirtualProduction/LiveLinkVCAM/vcam/SettingsViewController.swift b/Engine/Extras/VirtualProduction/LiveLinkVCAM/vcam/SettingsViewController.swift
--- a/Engine/Extras/VirtualProduction/LiveLinkVCAM/vcam/SettingsViewController.swift
+++ b/Engine/Extras/VirtualProduction/LiveLinkVCAM/vcam/SettingsViewController.swift
@@ -13,6 +13,25 @@ class SettingsViewController : UITableViewController {
     
     let appSettings = AppSettings.shared
     
+    class LocalizedConnectionType {
+        
+        class func title() -> String {
+            return NSLocalizedString("connectiontype-title", value: "ConnectionType", comment: "")
+        }
+        class func footer() -> String {
+            return ""
+            //return NSLocalizedString("connectiontype-footer", value: "", comment: "")
+        }
+        class func value(_ type : StreamingConnectionType) -> String {
+            switch type {
+            case .remoteSession:
+                return NSLocalizedString("connectiontype-value-remotesession", value: "Remote Session", comment: "")
+            case .webRTC:
+                return NSLocalizedString("connectiontype-value-pixelstreaming", value: "Pixel Streaming", comment: "")
+            }
+        }
+    }
+                    
     override func viewDidLoad() {
         super.viewDidLoad()
         if UIDevice.current.userInterfaceIdiom == .phone {
@@ -40,12 +59,7 @@ class SettingsViewController : UITableViewController {
             cell.detailTextLabel?.text = Timecode.sourceToString(appSettings.timecodeSourceEnum())
             
         case "connectionType":
-            switch appSettings.connectionTypeEnum() {
-            case .remoteSession:
-                cell.detailTextLabel?.text = "Remote Session"
-            case .webRTC:
-                cell.detailTextLabel?.text = "Pixel Streaming"
-            }
+            cell.detailTextLabel?.text = LocalizedConnectionType.value(appSettings.connectionTypeEnum())
 
         default:
             break
@@ -84,12 +98,12 @@ class SettingsViewController : UITableViewController {
             
             if segue.identifier == "connectionType" {
                 
-                vc.navigationItem.title = "Connection Type"
-                vc.items = [ "Remote Session", "Pixel Streaming" ]
-                vc.selectedIndex = appSettings.connectionType == "RemoteSession" ? 0 : 1
-                vc.footerString = "Description of Connection Types"
+                vc.navigationItem.title = LocalizedConnectionType.title()
+                vc.items = [ LocalizedConnectionType.value(.remoteSession), LocalizedConnectionType.value(.webRTC) ]
+                vc.selectedIndex = appSettings.connectionTypeEnum() == .remoteSession ? 0 : 1
+                vc.footerString = LocalizedConnectionType.footer()
                 vc.completion = { (index) in
-                    self.appSettings.connectionType = index == 0 ? "RemoteSession" : "WebRTC"
+                    self.appSettings.setConnectionTypeEnum(index == 0 ? .remoteSession : .webRTC)
                 }
             }
         }
diff --git a/Engine/Extras/VirtualProduction/LiveLinkVCAM/vcam/ja.lproj/Localizable.strings b/Engine/Extras/VirtualProduction/LiveLinkVCAM/vcam/ja.lproj/Localizable.strings
--- a/Engine/Extras/VirtualProduction/LiveLinkVCAM/vcam/ja.lproj/Localizable.strings
+++ b/Engine/Extras/VirtualProduction/LiveLinkVCAM/vcam/ja.lproj/Localizable.strings
@@ -157,3 +157,14 @@
 /* An unknown or unsupported value was found */
 "value-unknown" = "不明";
 
+/* Type of UE connection to use (remote session or pixel streaming) */
+"connectiontype-title" = "接続タイプ";
+
+/* Footer describing the connection type option */
+"connectiontype-footer" = "";
+
+/* Remote Session */
+"connectiontype-value-remotesession" = "リモート セッション";
+
+/* Pixel Streaming */
+"connectiontype-value-pixelstreaming" = "ピクセルストリーミング";
diff --git a/Engine/Extras/VirtualProduction/LiveLinkVCAM/vcam/ja.lproj/Main.strings b/Engine/Extras/VirtualProduction/LiveLinkVCAM/vcam/ja.lproj/Main.strings
--- a/Engine/Extras/VirtualProduction/LiveLinkVCAM/vcam/ja.lproj/Main.strings
+++ b/Engine/Extras/VirtualProduction/LiveLinkVCAM/vcam/ja.lproj/Main.strings
@@ -97,3 +97,6 @@
 /* Class = "UINavigationItem"; title = "License"; ObjectID = "ZHG-A7-gPg"; */
 "ZHG-A7-gPg.title" = "ライセンス";
 
+/* Class = "UILabel"; title = "Connection Type"; ObjectID = "DJe-a8-JcI"; */
+"DJe-a8-JcI.text" = "接続タイプ";
+
diff --git a/Engine/Extras/VirtualProduction/LiveLinkVCAM/vcam/ko.lproj/Localizable.strings b/Engine/Extras/VirtualProduction/LiveLinkVCAM/vcam/ko.lproj/Localizable.strings
--- a/Engine/Extras/VirtualProduction/LiveLinkVCAM/vcam/ko.lproj/Localizable.strings
+++ b/Engine/Extras/VirtualProduction/LiveLinkVCAM/vcam/ko.lproj/Localizable.strings
@@ -157,3 +157,14 @@
 /* An unknown or unsupported value was found */
 "value-unknown" = "알 수 없음";
 
+/* Type of UE connection to use (remote session or pixel streaming) */
+"connectiontype-title" = "연결 타입";
+
+/* Footer describing the connection type option */
+"connectiontype-footer" = "";
+
+/* Remote Session */
+"connectiontype-value-remotesession" = "리모트 세션";
+
+/* Pixel Streaming */
+"connectiontype-value-pixelstreaming" = "픽셀 스트리밍";
diff --git a/Engine/Extras/VirtualProduction/LiveLinkVCAM/vcam/ko.lproj/Main.strings b/Engine/Extras/VirtualProduction/LiveLinkVCAM/vcam/ko.lproj/Main.strings
--- a/Engine/Extras/VirtualProduction/LiveLinkVCAM/vcam/ko.lproj/Main.strings
+++ b/Engine/Extras/VirtualProduction/LiveLinkVCAM/vcam/ko.lproj/Main.strings
@@ -97,3 +97,5 @@
 /* Class = "UINavigationItem"; title = "License"; ObjectID = "ZHG-A7-gPg"; */
 "ZHG-A7-gPg.title" = "라이선스";
 
+/* Class = "UILabel"; title = "Connection Type"; ObjectID = "DJe-a8-JcI"; */
+"DJe-a8-JcI.text" = "연결 타입";
diff --git a/Engine/Extras/VirtualProduction/LiveLinkVCAM/vcam/zh-Hans.lproj/Localizable.strings b/Engine/Extras/VirtualProduction/LiveLinkVCAM/vcam/zh-Hans.lproj/Localizable.strings
--- a/Engine/Extras/VirtualProduction/LiveLinkVCAM/vcam/zh-Hans.lproj/Localizable.strings
+++ b/Engine/Extras/VirtualProduction/LiveLinkVCAM/vcam/zh-Hans.lproj/Localizable.strings
@@ -157,3 +157,15 @@
 /* An unknown or unsupported value was found */
 "value-unknown" = "未知";
 
+/* Type of UE connection to use (remote session or pixel streaming) */
+"connectiontype-title" = "接続タイプ";
+
+/* Footer describing the connection type option */
+"connectiontype-footer" = "";
+
+/* Remote Session */
+"connectiontype-value-remotesession" = "远程会话";
+
+/* Pixel Streaming */
+"connectiontype-value-pixelstreaming" = "像素流送";
+
diff --git a/Engine/Extras/VirtualProduction/LiveLinkVCAM/vcam/zh-Hans.lproj/Main.strings b/Engine/Extras/VirtualProduction/LiveLinkVCAM/vcam/zh-Hans.lproj/Main.strings
--- a/Engine/Extras/VirtualProduction/LiveLinkVCAM/vcam/zh-Hans.lproj/Main.strings
+++ b/Engine/Extras/VirtualProduction/LiveLinkVCAM/vcam/zh-Hans.lproj/Main.strings
@@ -97,3 +97,5 @@
 /* Class = "UINavigationItem"; title = "License"; ObjectID = "ZHG-A7-gPg"; */
 "ZHG-A7-gPg.title" = "许可证";
 
+/* Class = "UILabel"; title = "Connection Type"; ObjectID = "DJe-a8-JcI"; */
+"DJe-a8-JcI.text" = "连接类型";
