Fixed bug where the touch events were still being passed through while the AR coaching view was visible.

diff --git a/Engine/Extras/VirtualProduction/LiveLinkVCAM/vcam/RemoteSessionStreamingConnection.swift b/Engine/Extras/VirtualProduction/LiveLinkVCAM/vcam/RemoteSessionStreamingConnection.swift
--- a/Engine/Extras/VirtualProduction/LiveLinkVCAM/vcam/RemoteSessionStreamingConnection.swift
+++ b/Engine/Extras/VirtualProduction/LiveLinkVCAM/vcam/RemoteSessionStreamingConnection.swift
@@ -79,6 +79,13 @@ class RemoteSessionStreamingConnection : StreamingConnection  {
         }
     }
     
+    override var relayTouchEvents : Bool {
+        willSet {
+            self.remoteSessionView?.endAllTouches()
+        }
+    }
+
+    
     required init(subjectName: String) {
         super.init(subjectName: subjectName)
 
@@ -139,7 +146,6 @@ class RemoteSessionStreamingConnection : StreamingConnection  {
     }
     
     override func reconnect() {
-        
         oscConnection?.reconnect()
     }
     
@@ -169,7 +175,6 @@ class RemoteSessionStreamingConnection : StreamingConnection  {
    override func sendControllerButtonReleased(_ type : StreamingConnectionControllerInputType, controllerIndex : UInt8) {
        self.oscConnection?.send(.controllerButtonReleased, arguments: [ OSCArgument.blob(OSCUtility.ueControllerButtonData(key: type.rawValue, controller: Int(controllerIndex), isRepeat: false)) ] )
     }
-
     
     func restartLiveLink() {
 
@@ -224,6 +229,7 @@ extension RemoteSessionStreamingConnection : RemoteSessionViewDelegate {
     
     func remoteSessionView(_ view: RemoteSessionView?, touch type: StreamingConnectionTouchType, index: Int, at point: CGPoint, force: CGFloat) {
         
+        guard self.relayTouchEvents else { return }
         guard let rsv = remoteSessionView else { return }
         
         let normalizedPoint = CGPoint(x: point.x / rsv.frame.size.width, y: point.y / rsv.frame.size.height)
diff --git a/Engine/Extras/VirtualProduction/LiveLinkVCAM/vcam/RemoteSessionView.swift b/Engine/Extras/VirtualProduction/LiveLinkVCAM/vcam/RemoteSessionView.swift
--- a/Engine/Extras/VirtualProduction/LiveLinkVCAM/vcam/RemoteSessionView.swift
+++ b/Engine/Extras/VirtualProduction/LiveLinkVCAM/vcam/RemoteSessionView.swift
@@ -54,8 +54,6 @@ class RemoteSessionView : UIView {
     
     override func touchesBegan(_ touches: Set<UITouch>, with event: UIEvent?) {
 
-        //guard relayTouchEvents else { return }
-        
         // for each touch, we need to maintain an index to send via OSC, so that UE
         // can correctly map between began/moved. The currentTouches array will maintain
         // the indices for each touch (keeping some nil if needed) and clean up / compress
@@ -83,8 +81,6 @@ class RemoteSessionView : UIView {
     
     override func touchesMoved(_ touches: Set<UITouch>, with event: UIEvent?) {
 
-        //guard relayTouchEvents else { return }
-
         for touch in touches {
             if let fingerIndex = _currentTouches.firstIndex(of: touch) {
                 self.delegate?.remoteSessionView(self, touch: .moved, index: fingerIndex, at: touch.location(in: self), force: touch.force)
@@ -95,8 +91,6 @@ class RemoteSessionView : UIView {
     
     override func touchesEnded(_ touches: Set<UITouch>, with event: UIEvent?) {
 
-        //guard relayTouchEvents else { return }
-
         for touch in touches {
             
             
@@ -117,4 +111,13 @@ class RemoteSessionView : UIView {
             _currentTouches.removeAll()
         }
     }
+    
+    func endAllTouches() {
+        for i in 0..<_currentTouches.count {
+            if let t = _currentTouches[i] {
+                self.delegate?.remoteSessionView(self, touch: .ended, index: i, at: t.location(in: self), force: t.force)
+            }
+        }
+        _currentTouches.removeAll()
+    }
 }
diff --git a/Engine/Extras/VirtualProduction/LiveLinkVCAM/vcam/StreamingConnection.swift b/Engine/Extras/VirtualProduction/LiveLinkVCAM/vcam/StreamingConnection.swift
--- a/Engine/Extras/VirtualProduction/LiveLinkVCAM/vcam/StreamingConnection.swift
+++ b/Engine/Extras/VirtualProduction/LiveLinkVCAM/vcam/StreamingConnection.swift
@@ -87,6 +87,11 @@ class StreamingConnection : NSObject {
         }
     }
     
+    var relayTouchEvents = true {
+        willSet {
+        }
+    }
+
     var stats : StreamingConnectionStats?
     
     private(set) var videoAspectRatio: CGFloat?
diff --git a/Engine/Extras/VirtualProduction/LiveLinkVCAM/vcam/VideoViewController+ARCoachingOverlayViewDelegate.swift b/Engine/Extras/VirtualProduction/LiveLinkVCAM/vcam/VideoViewController+ARCoachingOverlayViewDelegate.swift
--- a/Engine/Extras/VirtualProduction/LiveLinkVCAM/vcam/VideoViewController+ARCoachingOverlayViewDelegate.swift
+++ b/Engine/Extras/VirtualProduction/LiveLinkVCAM/vcam/VideoViewController+ARCoachingOverlayViewDelegate.swift
@@ -13,7 +13,7 @@ extension VideoViewController : ARCoachingOverlayViewDelegate
 {
     func coachingOverlayViewWillActivate(_ coachingOverlayView: ARCoachingOverlayView) {
         
-        self.relayTouchEvents = false
+        self.streamingConnection?.relayTouchEvents = false
 
     }
 
@@ -25,7 +25,7 @@ extension VideoViewController : ARCoachingOverlayViewDelegate
      */
     func coachingOverlayViewDidDeactivate(_ coachingOverlayView: ARCoachingOverlayView) {
         
-        self.relayTouchEvents = true
+        self.streamingConnection?.relayTouchEvents = true
         
     }
 
diff --git a/Engine/Extras/VirtualProduction/LiveLinkVCAM/vcam/VideoViewController.swift b/Engine/Extras/VirtualProduction/LiveLinkVCAM/vcam/VideoViewController.swift
--- a/Engine/Extras/VirtualProduction/LiveLinkVCAM/vcam/VideoViewController.swift
+++ b/Engine/Extras/VirtualProduction/LiveLinkVCAM/vcam/VideoViewController.swift
@@ -64,27 +64,6 @@ class VideoViewController : BaseViewController {
         }
     }
     
-    var relayTouchEvents = true {
-        didSet {
-
-            // if we are stopping touch events sent to OSC, then we need to send an "ended" event
-            // for all of the current touches so the UI doesn't think it's still occurring
-            /*
-            if relayTouchEvents == false {
-
-                for i in 0..<self.currentTouches.count {
-                    if let t = self.currentTouches[i] {
-                        self.sendTouch(.ended, point: t.location(in: self.renderView), finger: i, force: 0);
-                    }
-                }
-                self.currentTouches.removeAll()
-
-            }
-             */
-            
-        }
-    }
-
     @IBOutlet weak var arView : ARSCNView!
 
     @IBOutlet weak var reconnectingBlurView : UIVisualEffectView!
diff --git a/Engine/Extras/VirtualProduction/LiveLinkVCAM/vcam/WebRTC/Extensions/TouchControls.swift b/Engine/Extras/VirtualProduction/LiveLinkVCAM/vcam/WebRTC/Extensions/TouchControls.swift
--- a/Engine/Extras/VirtualProduction/LiveLinkVCAM/vcam/WebRTC/Extensions/TouchControls.swift
+++ b/Engine/Extras/VirtualProduction/LiveLinkVCAM/vcam/WebRTC/Extensions/TouchControls.swift
@@ -44,6 +44,19 @@ final class TouchControls : TouchDelegate {
     var fingers : [Int] = [10, 9, 8, 7, 6, 5, 4, 3, 2, 1, 0]
     var fingerIds = [UITouch : Int]()
     
+    var relayTouchEvents = true {
+        willSet {
+            
+            // cancel/end all touches
+            var touches = Set<UITouch>()
+            for t in self.fingerIds {
+                touches.insert(t.key)
+            }
+            touchesCancelled(touches)
+        }
+
+    }
+    
     init(_ webRTCClient : WebRTCClient, touchView: UIView) {
         self.webRTCClient = webRTCClient
         self.touchView = touchView
@@ -101,7 +114,7 @@ final class TouchControls : TouchDelegate {
         // same as touches ended unless we are hooking up gestures
         touchesEnded(touches)
     }
-    
+
     func onVideoChangedSize(size: CGSize) {
         self.videoAspectRatio = size.height / size.width
     }
@@ -131,6 +144,8 @@ final class TouchControls : TouchDelegate {
     
     func sendTouchData(messageType: PixelStreamingToStreamerMessage, touches: Set<UITouch>) {
         
+        guard relayTouchEvents else { return }
+        
         var bytes: [UInt8] = []
         
         // Add type of touch using 1 byte
diff --git a/Engine/Extras/VirtualProduction/LiveLinkVCAM/vcam/WebRTCStreamingConnection.swift b/Engine/Extras/VirtualProduction/LiveLinkVCAM/vcam/WebRTCStreamingConnection.swift
--- a/Engine/Extras/VirtualProduction/LiveLinkVCAM/vcam/WebRTCStreamingConnection.swift
+++ b/Engine/Extras/VirtualProduction/LiveLinkVCAM/vcam/WebRTCStreamingConnection.swift
@@ -47,6 +47,12 @@ class WebRTCStreamingConnection : StreamingConnection {
             return (self.webRTCClientState ?? .disconnected) == .connected
         }
     }
+    
+    override var relayTouchEvents: Bool {
+        willSet {
+            self.touchControls?.relayTouchEvents = relayTouchEvents
+        }
+    }
 
     override var renderView: UIView? {
         didSet {
