Fixed touch points when the delivered video aspect ratio doesn't match the screen aspect ratio.

diff --git a/Engine/Extras/VirtualProduction/LiveLinkVCAM/vcam/WebRTC/Extensions/TouchControls.swift b/Engine/Extras/VirtualProduction/LiveLinkVCAM/vcam/WebRTC/Extensions/TouchControls.swift
--- a/Engine/Extras/VirtualProduction/LiveLinkVCAM/vcam/WebRTC/Extensions/TouchControls.swift
+++ b/Engine/Extras/VirtualProduction/LiveLinkVCAM/vcam/WebRTC/Extensions/TouchControls.swift
@@ -116,20 +116,27 @@ final class TouchControls : TouchDelegate {
     }
 
     func onVideoChangedSize(size: CGSize) {
-        self.videoAspectRatio = size.height / size.width
+        self.videoAspectRatio = size.width / size.height
     }
     
     // Unsigned XY positions are normalized to the ratio (0.0..1.0) along a viewport axis and then quantized into an uint16 (0..65536).
     func normalizeAndQuantize(_ touch: UITouch) -> TouchPoint {
-        let uiAspectRatio : CGFloat = touchView.bounds.height / touchView.bounds.width
-        let uiBigger : Bool = uiAspectRatio > videoAspectRatio
-        let ratio : CGFloat = uiAspectRatio / videoAspectRatio;
+        let uiAspectRatio : CGFloat = touchView.bounds.width / touchView.bounds.height
         let touchLocation : CGPoint = touch.location(in: touchView)
         let viewBounds : CGRect = touchView.bounds
         
         // normalise x,y to the UI element bounds
-        let normalizedX : CGFloat = uiBigger ? (touchLocation.x / viewBounds.width) : (ratio * (touchLocation.x / viewBounds.width - 0.5) + 0.5)
-        let normalizedY : CGFloat = uiBigger ? (ratio * (touchLocation.y / viewBounds.height - 0.5) + 0.5) : (touchLocation.y / viewBounds.height)
+        var normalizedX : CGFloat = touchLocation.x / viewBounds.width
+        if uiAspectRatio > videoAspectRatio {
+            normalizedX = (normalizedX - 0.5) * (uiAspectRatio / videoAspectRatio) + 0.5
+        }
+        
+        var normalizedY : CGFloat = touchLocation.y / viewBounds.height
+        if uiAspectRatio < videoAspectRatio {
+            normalizedY = (normalizedY - 0.5) * (videoAspectRatio / uiAspectRatio) + 0.5
+        }
+        
+        Log.info("\(normalizedX) \(normalizedY)")
         
         // normalize force value of touch
         let normalizedForce : UInt8 = touch.maximumPossibleForce > 0 ? UInt8(touch.force / touch.maximumPossibleForce * CGFloat(UInt8.max)) : 1
