Live Link VCAM v1.3.0 build 10 (release).  Last localization.

diff --git a/Engine/Extras/VirtualProduction/LiveLinkVCAM/vcam/SettingsViewController.swift b/Engine/Extras/VirtualProduction/LiveLinkVCAM/vcam/SettingsViewController.swift
--- a/Engine/Extras/VirtualProduction/LiveLinkVCAM/vcam/SettingsViewController.swift
+++ b/Engine/Extras/VirtualProduction/LiveLinkVCAM/vcam/SettingsViewController.swift
@@ -19,8 +19,7 @@ class SettingsViewController : UITableViewController {
             return NSLocalizedString("connectiontype-title", value: "ConnectionType", comment: "")
         }
         class func footer() -> String {
-            return ""
-            //return NSLocalizedString("connectiontype-footer", value: "", comment: "")
+            return NSLocalizedString("connectiontype-footer", value: "Unreal Engine Pixel Streaming introduced editor support in Unreal Engine 5.1 and is the recommended connection type for Live Link VCAM going forward.  The legacy Remote Session connection type will continue to be supported in Live Link VCAM for older projects.", comment: "")
         }
         class func value(_ type : StreamingConnectionType) -> String {
             switch type {
diff --git a/Engine/Extras/VirtualProduction/LiveLinkVCAM/vcam/ja.lproj/Localizable.strings b/Engine/Extras/VirtualProduction/LiveLinkVCAM/vcam/ja.lproj/Localizable.strings
--- a/Engine/Extras/VirtualProduction/LiveLinkVCAM/vcam/ja.lproj/Localizable.strings
+++ b/Engine/Extras/VirtualProduction/LiveLinkVCAM/vcam/ja.lproj/Localizable.strings
@@ -161,7 +161,7 @@
 "connectiontype-title" = "接続タイプ";
 
 /* Footer describing the connection type option */
-"connectiontype-footer" = "";
+"connectiontype-footer" = "Unreal Engine Pixel Streaming が Unreal Engine 5.1 でエディタでサポートされるようになり、この接続タイプが Live Link VCAM で推奨されるものになりました。従来のリモート セッション 接続タイプも過去のプロジェクトのために引き続き Live Link VCAM ではサポートされます。";
 
 /* Remote Session */
 "connectiontype-value-remotesession" = "リモート セッション";
diff --git a/Engine/Extras/VirtualProduction/LiveLinkVCAM/vcam/ko.lproj/Localizable.strings b/Engine/Extras/VirtualProduction/LiveLinkVCAM/vcam/ko.lproj/Localizable.strings
--- a/Engine/Extras/VirtualProduction/LiveLinkVCAM/vcam/ko.lproj/Localizable.strings
+++ b/Engine/Extras/VirtualProduction/LiveLinkVCAM/vcam/ko.lproj/Localizable.strings
@@ -161,7 +161,7 @@
 "connectiontype-title" = "연결 타입";
 
 /* Footer describing the connection type option */
-"connectiontype-footer" = "";
+"connectiontype-footer" = "이제 언리얼 엔진 픽셀 스트리밍이 Live Link  VCAM에서 권장되는 연결 유형으로 언리얼 엔진 5.1 에디터에서 지원됩니다. 이전 프로젝트의 경우 기존 원격 세션 연결 유형은 Live Link VCAM에서 계속 지원됩니다.";
 
 /* Remote Session */
 "connectiontype-value-remotesession" = "리모트 세션";
diff --git a/Engine/Extras/VirtualProduction/LiveLinkVCAM/vcam/zh-Hans.lproj/Localizable.strings b/Engine/Extras/VirtualProduction/LiveLinkVCAM/vcam/zh-Hans.lproj/Localizable.strings
--- a/Engine/Extras/VirtualProduction/LiveLinkVCAM/vcam/zh-Hans.lproj/Localizable.strings
+++ b/Engine/Extras/VirtualProduction/LiveLinkVCAM/vcam/zh-Hans.lproj/Localizable.strings
@@ -161,7 +161,7 @@
 "connectiontype-title" = "接続タイプ";
 
 /* Footer describing the connection type option */
-"connectiontype-footer" = "";
+"connectiontype-footer" = "在虚幻引擎5.1中，编辑器将支持像素流送功能，并将成为Live Link VCAM的推荐连接类型。但Live Link VCAM仍将继续支持旧版的远程会话连接类型，以便用于旧项目的开发。";
 
 /* Remote Session */
 "connectiontype-value-remotesession" = "远程会话";
