Live Link VCAM 1.3.0 build 11.

diff --git a/Engine/Extras/VirtualProduction/LiveLinkVCAM/vcam/RemoteSessionStreamingConnection.swift b/Engine/Extras/VirtualProduction/LiveLinkVCAM/vcam/RemoteSessionStreamingConnection.swift
--- a/Engine/Extras/VirtualProduction/LiveLinkVCAM/vcam/RemoteSessionStreamingConnection.swift
+++ b/Engine/Extras/VirtualProduction/LiveLinkVCAM/vcam/RemoteSessionStreamingConnection.swift
@@ -80,8 +80,10 @@ class RemoteSessionStreamingConnection : StreamingConnection  {
     }
     
     override var relayTouchEvents : Bool {
-        willSet {
-            self.remoteSessionView?.endAllTouches()
+        didSet {
+            if !relayTouchEvents {
+                self.remoteSessionView?.endAllTouches()
+            }
         }
     }
 
diff --git a/Engine/Extras/VirtualProduction/LiveLinkVCAM/vcam/StreamingConnection.swift b/Engine/Extras/VirtualProduction/LiveLinkVCAM/vcam/StreamingConnection.swift
--- a/Engine/Extras/VirtualProduction/LiveLinkVCAM/vcam/StreamingConnection.swift
+++ b/Engine/Extras/VirtualProduction/LiveLinkVCAM/vcam/StreamingConnection.swift
@@ -88,7 +88,7 @@ class StreamingConnection : NSObject {
     }
     
     var relayTouchEvents = true {
-        willSet {
+        didSet {
         }
     }
 
diff --git a/Engine/Extras/VirtualProduction/LiveLinkVCAM/vcam/WebRTC/Extensions/TouchControls.swift b/Engine/Extras/VirtualProduction/LiveLinkVCAM/vcam/WebRTC/Extensions/TouchControls.swift
--- a/Engine/Extras/VirtualProduction/LiveLinkVCAM/vcam/WebRTC/Extensions/TouchControls.swift
+++ b/Engine/Extras/VirtualProduction/LiveLinkVCAM/vcam/WebRTC/Extensions/TouchControls.swift
@@ -45,14 +45,16 @@ final class TouchControls : TouchDelegate {
     var fingerIds = [UITouch : Int]()
     
     var relayTouchEvents = true {
-        willSet {
+        didSet {
             
-            // cancel/end all touches
-            var touches = Set<UITouch>()
-            for t in self.fingerIds {
-                touches.insert(t.key)
+            if !relayTouchEvents {
+                // cancel/end all touches
+                var touches = Set<UITouch>()
+                for t in self.fingerIds {
+                    touches.insert(t.key)
+                }
+                touchesCancelled(touches)
             }
-            touchesCancelled(touches)
         }
 
     }
diff --git a/Engine/Extras/VirtualProduction/LiveLinkVCAM/vcam/WebRTCStreamingConnection.swift b/Engine/Extras/VirtualProduction/LiveLinkVCAM/vcam/WebRTCStreamingConnection.swift
--- a/Engine/Extras/VirtualProduction/LiveLinkVCAM/vcam/WebRTCStreamingConnection.swift
+++ b/Engine/Extras/VirtualProduction/LiveLinkVCAM/vcam/WebRTCStreamingConnection.swift
@@ -49,7 +49,7 @@ class WebRTCStreamingConnection : StreamingConnection {
     }
     
     override var relayTouchEvents: Bool {
-        willSet {
+        didSet {
             self.touchControls?.relayTouchEvents = relayTouchEvents
         }
     }
