TextureShare: Added missing filters for TextureShare SDK\Samples in InstalledEngineFilters.xml and updated binaries accordingly

diff --git a/Engine/Build/InstalledEngineFilters.xml b/Engine/Build/InstalledEngineFilters.xml
--- a/Engine/Build/InstalledEngineFilters.xml
+++ b/Engine/Build/InstalledEngineFilters.xml
@@ -284,6 +284,7 @@
 		Engine/Extras/MayaVelocityGridExporter/...
 		Engine/Extras/UnrealVS/...
 		Engine/Extras/Redist/en-us/*
+		Engine/Extras/VirtualProduction/TextureShare/...
 
 		<!-- Virtual Production Templates -->
 		Templates/TP_nDisplayBP/...
diff --git a/Engine/Extras/VirtualProduction/TextureShare/TextureShareSDK/Source/Serialize/TextureShareCoreSerialize.h b/Engine/Extras/VirtualProduction/TextureShare/TextureShareSDK/Source/Serialize/TextureShareCoreSerialize.h
--- a/Engine/Extras/VirtualProduction/TextureShare/TextureShareSDK/Source/Serialize/TextureShareCoreSerialize.h
+++ b/Engine/Extras/VirtualProduction/TextureShare/TextureShareSDK/Source/Serialize/TextureShareCoreSerialize.h
@@ -168,6 +168,14 @@ public:
 		});
 	}
 
+	template<typename SubT>
+	T* FindByEqualsFunc(const SubT& InValue)
+	{
+		return FindByPredicate([InValue](const T& DataIt) {
+			return DataIt.EqualsFunc(InValue);
+		});
+	}
+
 	template<typename SubT>
 	bool GetValuesByEqualsFunc(const SubT& InValue, TArray<T>& OutValues) const
 	{
