Fixed crash when an a poorly-formatted IP address was entered (spaces at the end or beginning,  bad characters, etc).

diff --git a/Engine/Extras/VirtualProduction/LiveLinkVCAM/vcam/StartViewController.swift b/Engine/Extras/VirtualProduction/LiveLinkVCAM/vcam/StartViewController.swift
--- a/Engine/Extras/VirtualProduction/LiveLinkVCAM/vcam/StartViewController.swift
+++ b/Engine/Extras/VirtualProduction/LiveLinkVCAM/vcam/StartViewController.swift
@@ -238,23 +238,31 @@ class StartViewController : BaseViewController {
             }
 
             do {
-                self.streamingConnection?.destination = self.ipAddress.text!
+                self.streamingConnection?.destination = self.ipAddress.text!.trimmed()
                 try self.streamingConnection?.connect()
                 
-            } catch {
+            } catch StreamingConnectionError.runtimeError(let errorMessage) {
                 
-                hideConnectingAlertView() {
+                showConnectionErrorAlert(errorMessage)
 
-                    let errorAlert = UIAlertController(title: Localized.titleError(), message: "Couldn't connect : \(error.localizedDescription)", preferredStyle: .alert)
-                    errorAlert.addAction(UIAlertAction(title: Localized.buttonOK(), style: .default, handler: { _ in
-                        self.hideConnectingView { }
-                    }))
-                    self.present(errorAlert, animated: true)
-                }
+            } catch {
+                
+                showConnectionErrorAlert("Couldn't connect : \(error.localizedDescription)")
             }
         }
     }
     
+    func showConnectionErrorAlert(_ message : String) {
+        hideConnectingAlertView() {
+
+            let errorAlert = UIAlertController(title: Localized.titleError(), message: "Couldn't connect : \(message)", preferredStyle: .alert)
+            errorAlert.addAction(UIAlertAction(title: Localized.buttonOK(), style: .default, handler: { _ in
+                self.hideConnectingView { }
+            }))
+            self.present(errorAlert, animated: true)
+        }
+    }
+    
     func hideConnectingView( _ completion : @escaping () -> Void) {
         UIView.animate(withDuration: 0.2, animations: {
             self.connectingView.alpha = 0.0
diff --git a/Engine/Extras/VirtualProduction/LiveLinkVCAM/vcam/StreamingConnection.swift b/Engine/Extras/VirtualProduction/LiveLinkVCAM/vcam/StreamingConnection.swift
--- a/Engine/Extras/VirtualProduction/LiveLinkVCAM/vcam/StreamingConnection.swift
+++ b/Engine/Extras/VirtualProduction/LiveLinkVCAM/vcam/StreamingConnection.swift
@@ -26,6 +26,12 @@ enum StreamingConnectionTouchType : String {
     case ended = "Ended"
 }
 
+enum StreamingConnectionError: Error {
+    case runtimeError(String)
+}
+
+
+
 enum StreamingConnectionControllerInputType : String {
     case thumbstickLeftX = "Gamepad_LeftX"
     case thumbstickLeftY = "Gamepad_LeftY"
diff --git a/Engine/Extras/VirtualProduction/LiveLinkVCAM/vcam/WebRTCStreamingConnection.swift b/Engine/Extras/VirtualProduction/LiveLinkVCAM/vcam/WebRTCStreamingConnection.swift
--- a/Engine/Extras/VirtualProduction/LiveLinkVCAM/vcam/WebRTCStreamingConnection.swift
+++ b/Engine/Extras/VirtualProduction/LiveLinkVCAM/vcam/WebRTCStreamingConnection.swift
@@ -38,7 +38,7 @@ class WebRTCStreamingConnection : StreamingConnection {
             self._url?.absoluteString ?? ""
         }
         set {
-            self._url = URL(string: "ws://\(newValue):80")!
+            self._url = URL(string: "ws://\(newValue):80")
         }
     }
     
@@ -132,15 +132,19 @@ class WebRTCStreamingConnection : StreamingConnection {
     
     override func connect() throws {
         
+        guard let url = self._url else {
+            throw StreamingConnectionError.runtimeError("The destination address wasn't formatted properly.")
+        }
+        
         if signalClient == nil {
             
             // We will use 3rd party library for websockets.
             let webSocketProvider: WebSocketProvider
             
             if #available(iOS 13.0, *) {
-                webSocketProvider = NativeWebSocketProvider(url: self._url!, timeout: 2.0)
+                webSocketProvider = NativeWebSocketProvider(url: url, timeout: 2.0)
             } else {
-                webSocketProvider = StarscreamWebSocket(url: self._url!)
+                webSocketProvider = StarscreamWebSocket(url: url)
             }
             
             self.signalClient = SignalingClient(webSocket: webSocketProvider)
