Enable LiveLink VCam to send the device resolution back to the Editor for use with overriding of the viewport resolution.

diff --git a/Engine/Extras/VirtualProduction/LiveLinkVCAM/vcam/WebRTC/Extensions/PixelStreamingToStreamerMessage.swift b/Engine/Extras/VirtualProduction/LiveLinkVCAM/vcam/WebRTC/Extensions/PixelStreamingToStreamerMessage.swift
--- a/Engine/Extras/VirtualProduction/LiveLinkVCAM/vcam/WebRTC/Extensions/PixelStreamingToStreamerMessage.swift
+++ b/Engine/Extras/VirtualProduction/LiveLinkVCAM/vcam/WebRTC/Extensions/PixelStreamingToStreamerMessage.swift
@@ -42,3 +42,18 @@ enum PixelStreamingToStreamerMessage : UInt8, Codable {
     
     case Transform = 100
 }
+
+struct PixelStreamingToStreamerResolutionCommand: Codable {
+    var resolution: Resolution
+    
+    struct Resolution: Codable {
+        var width: Int = 0
+        var height: Int = 0
+    }
+    
+    init(width: Int, height: Int){
+        self.resolution = Resolution(width: width, height: height)
+    }
+}
+
+
diff --git a/Engine/Extras/VirtualProduction/LiveLinkVCAM/vcam/WebRTCStreamingConnection.swift b/Engine/Extras/VirtualProduction/LiveLinkVCAM/vcam/WebRTCStreamingConnection.swift
--- a/Engine/Extras/VirtualProduction/LiveLinkVCAM/vcam/WebRTCStreamingConnection.swift
+++ b/Engine/Extras/VirtualProduction/LiveLinkVCAM/vcam/WebRTCStreamingConnection.swift
@@ -316,6 +316,36 @@ class WebRTCStreamingConnection : StreamingConnection {
             self.touchControls = TouchControls(webRTC, touchView: view)
             view.attachVideoTrack(track: track)
             view.attachTouchDelegate(delegate: self.touchControls!)
+            self.sendDeviceResolution()
+        }
+    }
+    
+    func sendDeviceResolution() {
+        if let client = self.webRTCClient {
+            let bounds = UIScreen.main.nativeBounds
+            
+            // the device nativeBounds are always in portrait orientation, so we swap them here because we always run
+            // in landscape.
+            let resolutionCommand = PixelStreamingToStreamerResolutionCommand(width: Int(bounds.height), height: Int(bounds.width))
+            
+            do {
+                let jsonData = try JSONEncoder().encode(resolutionCommand)
+                let jsonString = String(data: jsonData, encoding: .utf8)!
+                
+                var bytes: [UInt8] = []
+                // Append the message type
+                bytes.append(PixelStreamingToStreamerMessage.Command.rawValue)
+                bytes.append(contentsOf: UInt16(jsonString.count).toBytes())
+                // Append each character in the json string represented by its UTF16 character code
+                for char in jsonString {
+                    let charCode = char.utf16[char.utf16.index(char.utf16.startIndex, offsetBy: 0)]
+                    bytes.append(contentsOf: charCode.toBytes())
+                }
+                
+                client.sendData(Data(bytes))
+            } catch {
+                print(error)
+            }
         }
     }
 }
