Live Link VCAM v1.3.1 (build 3) : first pass at textbox input.

diff --git a/Engine/Extras/VirtualProduction/LiveLinkVCAM/vcam/StartViewController+StreamingConnectionDelegate.swift b/Engine/Extras/VirtualProduction/LiveLinkVCAM/vcam/StartViewController+StreamingConnectionDelegate.swift
--- a/Engine/Extras/VirtualProduction/LiveLinkVCAM/vcam/StartViewController+StreamingConnectionDelegate.swift
+++ b/Engine/Extras/VirtualProduction/LiveLinkVCAM/vcam/StartViewController+StreamingConnectionDelegate.swift
@@ -41,4 +41,7 @@ extension StartViewController : StreamingConnectionDelegate {
             }
         }
     }
+    
+    func streamingConnection(_ connection: StreamingConnection, requestsTextEditWithContents contents: String, handler: @escaping (Bool, String?) -> Void) {
+    }
 }
diff --git a/Engine/Extras/VirtualProduction/LiveLinkVCAM/vcam/StreamingConnection.swift b/Engine/Extras/VirtualProduction/LiveLinkVCAM/vcam/StreamingConnection.swift
--- a/Engine/Extras/VirtualProduction/LiveLinkVCAM/vcam/StreamingConnection.swift
+++ b/Engine/Extras/VirtualProduction/LiveLinkVCAM/vcam/StreamingConnection.swift
@@ -13,6 +13,7 @@ protocol StreamingConnectionDelegate : AnyObject {
     
     func streamingConnectionDidConnect(_ connection : StreamingConnection)
     func streamingConnection(_ connection : StreamingConnection, didDisconnectWithError err: Error?)
+    func streamingConnection(_ connection : StreamingConnection, requestsTextEditWithContents contents : String, handler : @escaping (Bool, String?) -> Void)
 }
 
 enum StreamingConnectionType : String {
diff --git a/Engine/Extras/VirtualProduction/LiveLinkVCAM/vcam/VideoViewController+StreamingConnectionDelegate.swift b/Engine/Extras/VirtualProduction/LiveLinkVCAM/vcam/VideoViewController+StreamingConnectionDelegate.swift
--- a/Engine/Extras/VirtualProduction/LiveLinkVCAM/vcam/VideoViewController+StreamingConnectionDelegate.swift
+++ b/Engine/Extras/VirtualProduction/LiveLinkVCAM/vcam/VideoViewController+StreamingConnectionDelegate.swift
@@ -7,6 +7,7 @@
 //
 
 import Foundation
+import UIKit
 
 extension VideoViewController : StreamingConnectionDelegate {
 
@@ -34,5 +35,23 @@ extension VideoViewController : StreamingConnectionDelegate {
         }
 
     }
+    
+    func streamingConnection(_ connection: StreamingConnection, requestsTextEditWithContents contents: String, handler: @escaping (Bool, String?) -> Void) {
+        
+        DispatchQueue.main.async {
+            let alert = UIAlertController(title: nil, message: contents.count == 0 ? "Enter the value:" : "Update the value:", preferredStyle: .alert)
+            alert.addTextField()
+            alert.textFields![0].text = contents
+            alert.addAction(UIAlertAction(title: Localized.buttonOK(), style: .default) {_ in
+                handler(true, alert.textFields![0].text!)
+            })
+            alert.addAction(UIAlertAction(title: Localized.buttonCancel(), style: .cancel) {_ in
+                handler(false, nil)
+            })
+
+            self.present(alert, animated: true)
+        }
+    }
+
 }
 
diff --git a/Engine/Extras/VirtualProduction/LiveLinkVCAM/vcam/WebRTC/Extensions/KeyboardControls.swift b/Engine/Extras/VirtualProduction/LiveLinkVCAM/vcam/WebRTC/Extensions/KeyboardControls.swift
new file mode 100644
--- /dev/null
+++ b/Engine/Extras/VirtualProduction/LiveLinkVCAM/vcam/WebRTC/Extensions/KeyboardControls.swift
@@ -0,0 +1,32 @@
+//
+//  KeyboardControls.swift
+//  vcam
+//
+//  Created by William Belcher on 1/2/2023.
+//
+
+import Foundation
+
+protocol KeyboardDelegate: AnyObject {
+    func submitString(_ entry: String)
+}
+
+final class KeyboardControls: KeyboardDelegate {
+    let webRTCClient: WebRTCClient
+    
+    init(_ webRTCClient: WebRTCClient) {
+        self.webRTCClient = webRTCClient
+    }
+    
+    func submitString(_ entry: String) {
+        var bytes: [UInt8] = []
+        bytes.append(PixelStreamingToStreamerMessage.TextboxEntry.rawValue)
+        bytes.append(UInt8(entry.count))
+        
+        for char in entry {
+            bytes.append(contentsOf: char.utf16[char.utf16.index(char.utf16.startIndex, offsetBy: 0)].toBytes())
+        }
+        
+        self.webRTCClient.sendData(Data(bytes))
+    }
+}
diff --git a/Engine/Extras/VirtualProduction/LiveLinkVCAM/vcam/WebRTC/Extensions/PixelStreamingToClientMessage.swift b/Engine/Extras/VirtualProduction/LiveLinkVCAM/vcam/WebRTC/Extensions/PixelStreamingToClientMessage.swift
--- a/Engine/Extras/VirtualProduction/LiveLinkVCAM/vcam/WebRTC/Extensions/PixelStreamingToClientMessage.swift
+++ b/Engine/Extras/VirtualProduction/LiveLinkVCAM/vcam/WebRTC/Extensions/PixelStreamingToClientMessage.swift
@@ -24,3 +24,12 @@ enum PixelStreamingToClientMessage: UInt8, Codable {
 }
 
 
+// Represents the command json received as part of the 'Command' PixelStreamingToClientMessage
+struct PixelStreamingToClientCommand: Decodable {
+    let command: String
+}
+
+struct PixelStreamingToClientShowOnScreenKeyboardCommand: Decodable {
+    let showOnScreenKeyboard: Bool
+    let contents: String
+}
diff --git a/Engine/Extras/VirtualProduction/LiveLinkVCAM/vcam/WebRTC/Extensions/PixelStreamingToStreamerMessage.swift b/Engine/Extras/VirtualProduction/LiveLinkVCAM/vcam/WebRTC/Extensions/PixelStreamingToStreamerMessage.swift
--- a/Engine/Extras/VirtualProduction/LiveLinkVCAM/vcam/WebRTC/Extensions/PixelStreamingToStreamerMessage.swift
+++ b/Engine/Extras/VirtualProduction/LiveLinkVCAM/vcam/WebRTC/Extensions/PixelStreamingToStreamerMessage.swift
@@ -41,6 +41,7 @@ enum PixelStreamingToStreamerMessage : UInt8, Codable {
     case GamepadAnalog = 92
     
     case Transform = 100
+    case TextboxEntry = 101
 }
 
 struct PixelStreamingToStreamerResolutionCommand: Codable {
diff --git a/Engine/Extras/VirtualProduction/LiveLinkVCAM/vcam/WebRTCStreamingConnection.swift b/Engine/Extras/VirtualProduction/LiveLinkVCAM/vcam/WebRTCStreamingConnection.swift
--- a/Engine/Extras/VirtualProduction/LiveLinkVCAM/vcam/WebRTCStreamingConnection.swift
+++ b/Engine/Extras/VirtualProduction/LiveLinkVCAM/vcam/WebRTCStreamingConnection.swift
@@ -16,6 +16,7 @@ class WebRTCStreamingConnection : StreamingConnection {
     private var webRTCClientState : RTCIceConnectionState?
     private var signalClient: SignalingClient?
     private var touchControls: TouchControls?
+    private var keyboardControls: KeyboardControls?
     private var webRTCView : WebRTCView?
     private var rtcVideoTrack : RTCVideoTrack?
     
@@ -314,6 +315,7 @@ class WebRTCStreamingConnection : StreamingConnection {
     func attachVideoTrack() {
         if let webRTC = webRTCClient, let view = self.webRTCView, let track = self.rtcVideoTrack {
             self.touchControls = TouchControls(webRTC, touchView: view)
+            self.keyboardControls = KeyboardControls(webRTC)
             view.attachVideoTrack(track: track)
             view.attachTouchDelegate(delegate: self.touchControls!)
             self.sendDeviceResolution()
@@ -349,6 +351,7 @@ class WebRTCStreamingConnection : StreamingConnection {
         }
     }
 }
+
 extension WebRTCStreamingConnection : WebRTCViewDelegate {
     
     func webRTCView(_ view: WebRTCView, didChangeVideoSize size: CGSize) {
@@ -487,13 +490,33 @@ extension WebRTCStreamingConnection: WebRTCClientDelegate {
                 switch payloadType {
                 case .VideoEncoderAvgQP:
                     let qp : String? = String(data: data.dropFirst(), encoding: .utf16LittleEndian)
-                    Log.info("Quality = \(qp ?? "N/A")")
+                    //Log.info("Quality = \(qp ?? "N/A")")
+                case .Command:
+                    let command: String? = String(data: data.dropFirst(), encoding: .utf16LittleEndian)
+                    Log.info("command = \(command ?? "NULL")")
+                    if let commandData = command?.data(using: .utf8) {
+                        do {
+                            let commandJson: PixelStreamingToClientCommand = try JSONDecoder().decode(PixelStreamingToClientCommand.self, from: commandData)
+                            if commandJson.command == "onScreenKeyboard" {
+                                let showKeyboardJson: PixelStreamingToClientShowOnScreenKeyboardCommand = try JSONDecoder().decode(PixelStreamingToClientShowOnScreenKeyboardCommand.self, from: commandData)
+
+                                if showKeyboardJson.showOnScreenKeyboard {
+                                    self.delegate?.streamingConnection(self, requestsTextEditWithContents: showKeyboardJson.contents) { (success, newContents) in
+                                        guard let contents = newContents else { return }
+                                        if success {
+                                            self.keyboardControls?.submitString(contents)
+                                        }
+                                    }
+                                }
+                            }
+                        } catch {
+                            Log.error("An error occurred parsing the command `\(command ?? "<invalid>")` : \(error.localizedDescription)")
+                        }
+                    }
                 case .QualityControlOwnership:
                     fallthrough
                 case .Response:
                     fallthrough
-                case .Command:
-                    fallthrough
                 case .FreezeFrame:
                     fallthrough
                 case .UnfreezeFrame:
