Live Link VCAM v1.3.1 build 4 : TextBox entry for UE input.

diff --git a/Engine/Extras/VirtualProduction/LiveLinkVCAM/vcam/WebRTC/Extensions/KeyboardControls.swift b/Engine/Extras/VirtualProduction/LiveLinkVCAM/vcam/WebRTC/Extensions/KeyboardControls.swift
--- a/Engine/Extras/VirtualProduction/LiveLinkVCAM/vcam/WebRTC/Extensions/KeyboardControls.swift
+++ b/Engine/Extras/VirtualProduction/LiveLinkVCAM/vcam/WebRTC/Extensions/KeyboardControls.swift
@@ -24,9 +24,9 @@ final class KeyboardControls: KeyboardDelegate {
         bytes.append(UInt8(entry.count))
         
         for char in entry {
-            bytes.append(contentsOf: char.utf16[char.utf16.index(char.utf16.startIndex, offsetBy: 0)].toBytes())
+            bytes.append(contentsOf: char.utf16[char.utf16.index(char.utf16.startIndex, offsetBy: 0)].bigEndian.toBytes())
         }
-        
+
         self.webRTCClient.sendData(Data(bytes))
     }
 }
diff --git a/Engine/Extras/VirtualProduction/LiveLinkVCAM/vcam/WebRTC/Extensions/PixelStreamingToStreamerMessage.swift b/Engine/Extras/VirtualProduction/LiveLinkVCAM/vcam/WebRTC/Extensions/PixelStreamingToStreamerMessage.swift
--- a/Engine/Extras/VirtualProduction/LiveLinkVCAM/vcam/WebRTC/Extensions/PixelStreamingToStreamerMessage.swift
+++ b/Engine/Extras/VirtualProduction/LiveLinkVCAM/vcam/WebRTC/Extensions/PixelStreamingToStreamerMessage.swift
@@ -41,7 +41,7 @@ enum PixelStreamingToStreamerMessage : UInt8, Codable {
     case GamepadAnalog = 92
     
     case Transform = 100
-    case TextboxEntry = 101
+    case TextboxEntry = 52
 }
 
 struct PixelStreamingToStreamerResolutionCommand: Codable {
diff --git a/Engine/Extras/VirtualProduction/LiveLinkVCAM/vcam/WebRTC/Extensions/TouchControls.swift b/Engine/Extras/VirtualProduction/LiveLinkVCAM/vcam/WebRTC/Extensions/TouchControls.swift
--- a/Engine/Extras/VirtualProduction/LiveLinkVCAM/vcam/WebRTC/Extensions/TouchControls.swift
+++ b/Engine/Extras/VirtualProduction/LiveLinkVCAM/vcam/WebRTC/Extensions/TouchControls.swift
@@ -138,8 +138,6 @@ final class TouchControls : TouchDelegate {
             normalizedY = (normalizedY - 0.5) * (videoAspectRatio / uiAspectRatio) + 0.5
         }
         
-        Log.info("\(normalizedX) \(normalizedY)")
-        
         // normalize force value of touch
         let normalizedForce : UInt8 = touch.maximumPossibleForce > 0 ? UInt8(touch.force / touch.maximumPossibleForce * CGFloat(UInt8.max)) : 1
         
