Live Link VCAM v1.3.1 (build 5) : fixed sending of characters to UE.

diff --git a/Engine/Extras/VirtualProduction/LiveLinkVCAM/vcam/WebRTC/Extensions/KeyboardControls.swift b/Engine/Extras/VirtualProduction/LiveLinkVCAM/vcam/WebRTC/Extensions/KeyboardControls.swift
--- a/Engine/Extras/VirtualProduction/LiveLinkVCAM/vcam/WebRTC/Extensions/KeyboardControls.swift
+++ b/Engine/Extras/VirtualProduction/LiveLinkVCAM/vcam/WebRTC/Extensions/KeyboardControls.swift
@@ -21,10 +21,10 @@ final class KeyboardControls: KeyboardDelegate {
     func submitString(_ entry: String) {
         var bytes: [UInt8] = []
         bytes.append(PixelStreamingToStreamerMessage.TextboxEntry.rawValue)
-        bytes.append(UInt8(entry.count))
+        bytes.append(contentsOf: UInt16(entry.count).toBytes())
         
         for char in entry {
-            bytes.append(contentsOf: char.utf16[char.utf16.index(char.utf16.startIndex, offsetBy: 0)].bigEndian.toBytes())
+            bytes.append(contentsOf: char.utf16[char.utf16.index(char.utf16.startIndex, offsetBy: 0)].toBytes())
         }
 
         self.webRTCClient.sendData(Data(bytes))
