[EpicStageApp] Rework UI to match newest mockups from Erik

diff --git a/Engine/Extras/VirtualProduction/EpicStageApp/pubspec.yaml b/Engine/Extras/VirtualProduction/EpicStageApp/pubspec.yaml
--- a/Engine/Extras/VirtualProduction/EpicStageApp/pubspec.yaml
+++ b/Engine/Extras/VirtualProduction/EpicStageApp/pubspec.yaml
@@ -86,6 +86,9 @@ flutter:
     - family: Droid Sans Mono
       fonts:
         - asset: assets/fonts/DroidSansMono.ttf
+    - family: Inter
+      fonts:
+        - asset: assets/fonts/Inter-VariableFont_slnt,wght.ttf
   #
   # For details regarding fonts from package dependencies,
   # see https://flutter.dev/custom-fonts/#from-packages
