Support for multi-stream vcam.

diff --git a/Engine/Extras/VirtualProduction/LiveLinkVCAM/vcam/StartViewController+StreamingConnectionDelegate.swift b/Engine/Extras/VirtualProduction/LiveLinkVCAM/vcam/StartViewController+StreamingConnectionDelegate.swift
--- a/Engine/Extras/VirtualProduction/LiveLinkVCAM/vcam/StartViewController+StreamingConnectionDelegate.swift
+++ b/Engine/Extras/VirtualProduction/LiveLinkVCAM/vcam/StartViewController+StreamingConnectionDelegate.swift
@@ -14,18 +14,22 @@ extension StartViewController : StreamingConnectionDelegate {
         Log.info("StreamingConnection \(connection.name) did connect to \(connection.destination)")
         
         DispatchQueue.main.async {
-            self.hideConnectingAlertView {
+            if self.connectingView.isHidden == true {
                 self.performSegue(withIdentifier: "showVideoView", sender: self)
-           }
+            } else {
+                self.hideConnectingView {
+                    self.performSegue(withIdentifier: "showVideoView", sender: self)
+                }
+            }
+            
         }
     }
 
     func streamingConnection(_ connection: StreamingConnection, didDisconnectWithError err: Error?) {
 
         DispatchQueue.main.async {
-
             connection.disconnect()
-            self.hideConnectingAlertView() {
+            self.hideConnectingAlertView {
                 if let e = err {
                     Log.info("StreamingConnection \(connection.name) disconnected with error : \(e.localizedDescription)")
 
@@ -44,4 +48,25 @@ extension StartViewController : StreamingConnectionDelegate {
     
     func streamingConnection(_ connection: StreamingConnection, requestsTextEditWithContents contents: String, handler: @escaping (Bool, String?) -> Void) {
     }
+    
+    func streamingConnection(_ connection: StreamingConnection, requestStreamerSelectionWithStreamers streamers: Array<String>, handler: @escaping (String) -> Void) {
+        self.pickerData = streamers;
+        self.selectedStreamer = streamers[0];
+        DispatchQueue.main.async {
+            self.hideConnectingView() {
+                let alert = UIAlertController(title: "Select Streamer", message: "\n\n\n\n\n\n", preferredStyle: .alert)
+                
+                let picker = UIPickerView(frame: CGRect(x: 5, y: 20, width: 250, height: 140))
+                picker.dataSource = self;
+                picker.delegate = self;
+                alert.view.addSubview(picker)
+
+                alert.addAction(UIAlertAction(title: "OK", style: .default) {_ in
+                    handler(self.selectedStreamer)
+                })
+                
+                self.present(alert, animated:true)
+            }
+        }
+    }
 }
diff --git a/Engine/Extras/VirtualProduction/LiveLinkVCAM/vcam/StartViewController.swift b/Engine/Extras/VirtualProduction/LiveLinkVCAM/vcam/StartViewController.swift
--- a/Engine/Extras/VirtualProduction/LiveLinkVCAM/vcam/StartViewController.swift
+++ b/Engine/Extras/VirtualProduction/LiveLinkVCAM/vcam/StartViewController.swift
@@ -12,7 +12,6 @@ import Kronos
 import GameController
 
 class StartViewController : BaseViewController {
-
     @IBOutlet weak var headerView : HeaderView!
     @IBOutlet weak var versionLabel : UILabel!
 
@@ -36,6 +35,9 @@ class StartViewController : BaseViewController {
     private var observers = [NSKeyValueObservation]()
     
     private var gameController : GCController?
+    
+    var pickerData: [String] = [String]()
+    var selectedStreamer: String = "";
 
     var ipAddressIsDemoMode : Bool {
         self.ipAddress.text == "demo.mode"
@@ -291,7 +293,6 @@ class StartViewController : BaseViewController {
             Log.info("gameControllerDidDisconnectNotification \(gc.vendorName ?? "Unknown controller")")
         }
     }
-
 }
 
 extension StartViewController : UIGestureRecognizerDelegate {
@@ -314,3 +315,24 @@ extension StartViewController : UITextFieldDelegate {
     }
 }
 
+extension StartViewController : UIPickerViewDelegate, UIPickerViewDataSource {
+    func numberOfComponents(in pickerView: UIPickerView) -> Int {
+        // Number of columns
+        return 1;
+    }
+    
+    func pickerView(_ pickerView: UIPickerView, numberOfRowsInComponent component: Int) -> Int {
+        // Number of rows
+        return pickerData.count
+    }
+    
+    func pickerView(_ pickerView: UIPickerView, titleForRow row: Int, forComponent component: Int) -> String? {
+        return pickerData[row]
+    }
+    
+    func pickerView(_ pickerView: UIPickerView, didSelectRow row: Int, inComponent component: Int) {
+            // This method is triggered whenever the user makes a change to the picker selection.
+            // The parameter named row and component represents what was selected.
+        selectedStreamer = pickerData[row]
+    }
+}
diff --git a/Engine/Extras/VirtualProduction/LiveLinkVCAM/vcam/StreamingConnection.swift b/Engine/Extras/VirtualProduction/LiveLinkVCAM/vcam/StreamingConnection.swift
--- a/Engine/Extras/VirtualProduction/LiveLinkVCAM/vcam/StreamingConnection.swift
+++ b/Engine/Extras/VirtualProduction/LiveLinkVCAM/vcam/StreamingConnection.swift
@@ -14,6 +14,7 @@ protocol StreamingConnectionDelegate : AnyObject {
     func streamingConnectionDidConnect(_ connection : StreamingConnection)
     func streamingConnection(_ connection : StreamingConnection, didDisconnectWithError err: Error?)
     func streamingConnection(_ connection : StreamingConnection, requestsTextEditWithContents contents : String, handler : @escaping (Bool, String?) -> Void)
+    func streamingConnection(_ connection: StreamingConnection, requestStreamerSelectionWithStreamers streamers: Array<String>, handler: @escaping (String) -> Void)
 }
 
 enum StreamingConnectionType : String {
diff --git a/Engine/Extras/VirtualProduction/LiveLinkVCAM/vcam/VideoViewController+StreamingConnectionDelegate.swift b/Engine/Extras/VirtualProduction/LiveLinkVCAM/vcam/VideoViewController+StreamingConnectionDelegate.swift
--- a/Engine/Extras/VirtualProduction/LiveLinkVCAM/vcam/VideoViewController+StreamingConnectionDelegate.swift
+++ b/Engine/Extras/VirtualProduction/LiveLinkVCAM/vcam/VideoViewController+StreamingConnectionDelegate.swift
@@ -52,6 +52,8 @@ extension VideoViewController : StreamingConnectionDelegate {
             self.present(alert, animated: true)
         }
     }
-
+    
+    func streamingConnection(_ connection: StreamingConnection, requestStreamerSelectionWithStreamers streamers: Array<String>, handler: @escaping (String) -> Void) {
+    }
 }
 
diff --git a/Engine/Extras/VirtualProduction/LiveLinkVCAM/vcam/WebRTC/Services/Message.swift b/Engine/Extras/VirtualProduction/LiveLinkVCAM/vcam/WebRTC/Services/Message.swift
--- a/Engine/Extras/VirtualProduction/LiveLinkVCAM/vcam/WebRTC/Services/Message.swift
+++ b/Engine/Extras/VirtualProduction/LiveLinkVCAM/vcam/WebRTC/Services/Message.swift
@@ -13,6 +13,9 @@ enum Message {
     case candidate(IceCandidate)
     case config(PeerConnectionConfig)
     case playerCount(Int)
+    case requestStreamerList
+    case streamerList(Array<String>)
+    case subscribe(String)
 }
 
 extension Message: Codable {
@@ -30,6 +33,8 @@ extension Message: Codable {
             self = .candidate(try container.decode(IceCandidate.self, forKey: .candidate))
         case "playerCount":
             self = .playerCount(try container.decode(Int.self, forKey: .count))
+        case "streamerList":
+            self = .streamerList(try container.decode(Array<String>.self, forKey: .ids))
         default:
             throw DecodeError.unknownType
         }
@@ -50,6 +55,14 @@ extension Message: Codable {
         case .playerCount(let count):
             try container.encode("playerCount", forKey: .type)
             try container.encode(count, forKey: .count)
+        case .streamerList(let streamerList):
+            try container.encode("streamerList", forKey: .type)
+            try container.encode(streamerList, forKey: .ids)
+        case .requestStreamerList:
+            try container.encode("listStreamers", forKey: .type)
+        case .subscribe(let streamerId):
+            try container.encode("subscribe", forKey: .type)
+            try container.encode(streamerId, forKey: .streamerId)
         }
     }
     
@@ -58,6 +71,6 @@ extension Message: Codable {
     }
     
     enum CodingKeys: String, CodingKey {
-        case type, sdp, candidate, peerConnectionOptions, count
+        case type, sdp, candidate, peerConnectionOptions, count, ids, requestStreamerList, streamerId
     }
 }
diff --git a/Engine/Extras/VirtualProduction/LiveLinkVCAM/vcam/WebRTC/Services/SignalingClient.swift b/Engine/Extras/VirtualProduction/LiveLinkVCAM/vcam/WebRTC/Services/SignalingClient.swift
--- a/Engine/Extras/VirtualProduction/LiveLinkVCAM/vcam/WebRTC/Services/SignalingClient.swift
+++ b/Engine/Extras/VirtualProduction/LiveLinkVCAM/vcam/WebRTC/Services/SignalingClient.swift
@@ -16,6 +16,7 @@ protocol SignalClientDelegate: AnyObject {
     func signalClient(_ signalClient: SignalingClient, didReceiveRemoteSdp sdp: RTCSessionDescription)
     func signalClient(_ signalClient: SignalingClient, didReceiveCandidate candidate: RTCIceCandidate)
     func signalClient(_ signalClient: SignalingClient, didReceiveConfig config: RTCConfiguration)
+    func signalClient(_ signalClient: SignalingClient, didReceiveStreamerList streamerList: Array<String>)
 }
 
 final class SignalingClient {
@@ -33,6 +34,7 @@ final class SignalingClient {
     func connect() {
         self.webSocket.delegate = self
         self.webSocket.connect()
+        self.sendRequestStreamerList()
     }
     
     func close() {
@@ -63,6 +65,27 @@ final class SignalingClient {
             debugPrint("Warning: Could not encode candidate: \(error)")
         }
     }
+    func sendRequestStreamerList() {
+        let message = Message.requestStreamerList
+        do {
+            let dataMessage = try self.encoder.encode(message)
+            self.webSocket.send(data: dataMessage)
+        }
+        catch {
+            debugPrint("Warning: Could not encode: \(error)")
+        }
+    }
+    
+    func subscribe(_ streamerId : String) {
+        let message = Message.subscribe(streamerId)
+        do {
+            let dataMessage = try self.encoder.encode(message)
+            self.webSocket.send(data: dataMessage)
+        }
+        catch {
+            debugPrint("Warning: Could not encode: \(error)")
+        }
+    }
 }
 
 
@@ -114,6 +137,12 @@ extension SignalingClient: WebSocketProviderDelegate {
             self.delegate?.signalClient(self, didReceiveRemoteSdp: sessionDescription.rtcSessionDescription)
         case .playerCount(let playerCount):
             print("Got player count=\(playerCount)")
+        case .streamerList(let streamerList):
+            self.delegate?.signalClient(self, didReceiveStreamerList: streamerList)
+        case .requestStreamerList:
+            break
+        case .subscribe(let streamerId):
+            break;
         }
     }
     
diff --git a/Engine/Extras/VirtualProduction/LiveLinkVCAM/vcam/WebRTCStreamingConnection.swift b/Engine/Extras/VirtualProduction/LiveLinkVCAM/vcam/WebRTCStreamingConnection.swift
--- a/Engine/Extras/VirtualProduction/LiveLinkVCAM/vcam/WebRTCStreamingConnection.swift
+++ b/Engine/Extras/VirtualProduction/LiveLinkVCAM/vcam/WebRTCStreamingConnection.swift
@@ -431,6 +431,18 @@ extension WebRTCStreamingConnection: SignalClientDelegate {
         }
     }
     
+    func signalClient(_ signalClient: SignalingClient, didReceiveStreamerList streamerList: Array<String>) {
+        // If we only have a single streamer, no need to show the selection dialogue
+        if streamerList.count == 1 {
+            signalClient.subscribe(streamerList[0])
+        } else if streamerList.count > 1 {
+            // Otherwise make sure we have more than 1 and display the picker
+            self.delegate?.streamingConnection(self, requestStreamerSelectionWithStreamers: streamerList) { (selectedStreamer) in
+                signalClient.subscribe(selectedStreamer)
+            }
+        }
+    }
+    
     func signalClientSendAnswer(_ signalClient: SignalingClient){
         if self.webRTCClient!.hasPeerConnnection() {
             Log.info("Sending answer sdp")
@@ -443,7 +455,6 @@ extension WebRTCStreamingConnection: SignalClientDelegate {
             Log.debug("WebRTC peer connection not setup yet - cannot handle sending answer.")
         }
     }
-    
 }
 
 extension WebRTCStreamingConnection: WebRTCClientDelegate {
