* Live Link VCAM iOS v1.3.1 (build 6)
* Added recently connected IP list as keyboard accessory.

diff --git a/Engine/Extras/VirtualProduction/LiveLinkVCAM/vcam/AppSettings.swift b/Engine/Extras/VirtualProduction/LiveLinkVCAM/vcam/AppSettings.swift
--- a/Engine/Extras/VirtualProduction/LiveLinkVCAM/vcam/AppSettings.swift
+++ b/Engine/Extras/VirtualProduction/LiveLinkVCAM/vcam/AppSettings.swift
@@ -25,7 +25,8 @@ class AppSettings : NSObject {
     private class Keys {
         static let didAcceptEULA = "didAcceptEULA"
         static let lastConnectionAddress = "lastConnectionAddress"
-        
+        static let recentConnectionAddresses = "recentConnectionAddresses"
+
         static let liveLinkSubjectName = "liveLinkSubjectName"
 
         static let timecodeSource = "timecodeSource"
@@ -43,6 +44,7 @@ class AppSettings : NSObject {
         return [
             Keys.didAcceptEULA : false,
             Keys.lastConnectionAddress : "",
+            Keys.recentConnectionAddresses : [],
 
             Keys.liveLinkSubjectName : defaultLiveLinkSubjectName(),
             
@@ -83,9 +85,29 @@ class AppSettings : NSObject {
             UserDefaults.standard.set(didAcceptEULA, forKey: Keys.didAcceptEULA)
         }
     }
+
+    func addRecentConnectionAddress(_ address : String) {
+        if let recents = UserDefaults.standard.array(forKey: Keys.recentConnectionAddresses) as? [String] {
+            // remove this connection address if it's already in the recent list and re-insert it at the front (most recent)
+            var newRecents = recents.filter { $0 != address }
+            if newRecents.count > 4 {
+                newRecents = Array(newRecents[...3])
+            }
+            newRecents.insert(address, at: 0)
+            UserDefaults.standard.set(newRecents, forKey: Keys.recentConnectionAddresses)
+        } else {
+            UserDefaults.standard.set([address], forKey: Keys.recentConnectionAddresses)
+        }
+    }
+    
+    var recentConnectionAddresses : [String]? {
+        get {
+            UserDefaults.standard.array(forKey: Keys.recentConnectionAddresses) as? [String]
+        }
+    }
     
     @objc dynamic var lastConnectionAddress : String = UserDefaults.standard.string(forKey: Keys.lastConnectionAddress) ?? "" {
-        didSet {
+        didSet{
             UserDefaults.standard.set(lastConnectionAddress, forKey: Keys.lastConnectionAddress)
         }
     }
diff --git a/Engine/Extras/VirtualProduction/LiveLinkVCAM/vcam/StartViewController.swift b/Engine/Extras/VirtualProduction/LiveLinkVCAM/vcam/StartViewController.swift
--- a/Engine/Extras/VirtualProduction/LiveLinkVCAM/vcam/StartViewController.swift
+++ b/Engine/Extras/VirtualProduction/LiveLinkVCAM/vcam/StartViewController.swift
@@ -62,8 +62,12 @@ class StartViewController : BaseViewController {
         NetUtility.triggerLocalNetworkPrivacyAlert()
         
         self.ipAddress.text = AppSettings.shared.lastConnectionAddress
+        self.ipAddress.inputAssistantItem.leadingBarButtonGroups.removeAll()
+        self.ipAddress.inputAssistantItem.trailingBarButtonGroups.removeAll()
         textFieldChanged(self.ipAddress)
         
+        self.rebuildRecentAddressesBarButtons()
+
         self.tapGesture = UITapGestureRecognizer(target: self, action: #selector(handleTap))
         self.tapGesture.cancelsTouchesInView = false
         self.tapGesture.delegate = self
@@ -167,6 +171,10 @@ class StartViewController : BaseViewController {
 
         if segue.identifier == "showVideoView" {
 
+            // connection was successful, we save the last address in our recents list
+            AppSettings.shared.addRecentConnectionAddress(AppSettings.shared.lastConnectionAddress)
+            self.rebuildRecentAddressesBarButtons()
+
             if let vc = segue.destination as? VideoViewController {
                 
                 // stop the timer locally which is sending LL identity xform
@@ -193,7 +201,7 @@ class StartViewController : BaseViewController {
         let connectButtonFrame = self.view.convert(connect.frame, from: connect.superview)
         
         if keyboardFrame.minY < connectButtonFrame.maxY {
-            self.entryViewYConstraint.constant = -(self.view.frame.height - keyboardFrame.size.height) / 2.0
+            self.entryViewYConstraint.constant = -keyboardFrame.size.height / 2.0
         } else {
             self.entryViewYConstraint.constant = 0
         }
@@ -213,6 +221,55 @@ class StartViewController : BaseViewController {
         }
     }
     
+    func rebuildRecentAddressesBarButtons() {
+
+        var view : UIView?
+
+        if let addresses = AppSettings.shared.recentConnectionAddresses {
+            for address in addresses {
+                
+                if view == nil {
+                    view = UIView()
+                }
+
+                let item = UIButton(configuration: UIButton.Configuration.gray())
+                item.setTitle(address, for: .normal)
+                item.setTitleColor(UIColor.white, for: .normal)
+                item.addTarget(self, action: #selector(handleRecentAddressSelection), for: .touchUpInside)
+                
+                view!.addSubview(item)
+                item.layoutToSuperview(.centerY)
+                
+                if view!.subviews.count == 1 {
+                    item.layoutToSuperview(.left)
+                } else {
+                    item.layout(.left, to: .right, of: view!.subviews[view!.subviews.count - 2], offset: 20)
+                }
+            }
+
+        }
+        
+        if let v = view {
+            v.subviews.last?.layoutToSuperview(.right)
+            let inputView = UIInputView(frame: CGRect(x: 0, y: 0, width: 100, height: 50), inputViewStyle: .keyboard)
+            inputView.addSubview(v)
+            v.layoutToSuperview(.top, offset: 4)
+            v.layoutToSuperview(.centerX, .bottom)
+            
+            self.ipAddress.inputAccessoryView = inputView  // inputAssistantItem.leadingBarButtonGroups.append(group)
+        }
+        
+    }
+    
+    @objc func handleRecentAddressSelection(_ sender : Any?) {
+        if let btn = sender as? UIButton {
+            if let addr = btn.title(for: .normal) {
+                self.ipAddress.text = addr
+                self.ipAddress.resignFirstResponder()
+            }
+        }
+    }
+    
     @objc func handleTap(_ recognizer: UITapGestureRecognizer) {
         
         self.view.endEditing(true)
