[EpicStageApp] Apply new styling to all dropdowns

diff --git a/Engine/Extras/VirtualProduction/EpicStageApp/pubspec.yaml b/Engine/Extras/VirtualProduction/EpicStageApp/pubspec.yaml
--- a/Engine/Extras/VirtualProduction/EpicStageApp/pubspec.yaml
+++ b/Engine/Extras/VirtualProduction/EpicStageApp/pubspec.yaml
@@ -19,6 +19,7 @@ version: 1.0.0+1
 
 environment:
   sdk: ">=2.17.0 <3.0.0"
+  flutter: "^3.7.0"
 
 # Dependencies specify other packages that your package needs in order to work.
 # To automatically upgrade your package dependencies to the latest versions
@@ -75,6 +76,9 @@ flutter:
     - assets/images/decoration/
     - assets/images/icons/
     - assets/images/map/
+  
+  shaders:
+    - assets/shaders/scroll_view_fade.frag
 
   # An image asset can refer to one or more resolution-specific "variants", see
   # https://flutter.dev/assets-and-images/#resolution-aware.
