Improved vertical centering of ui elements on iPhone when the keyboard appears.

diff --git a/Engine/Extras/VirtualProduction/LiveLinkVCAM/vcam/StartViewController.swift b/Engine/Extras/VirtualProduction/LiveLinkVCAM/vcam/StartViewController.swift
--- a/Engine/Extras/VirtualProduction/LiveLinkVCAM/vcam/StartViewController.swift
+++ b/Engine/Extras/VirtualProduction/LiveLinkVCAM/vcam/StartViewController.swift
@@ -201,7 +201,7 @@ class StartViewController : BaseViewController {
         let connectButtonFrame = self.view.convert(connect.frame, from: connect.superview)
         
         if keyboardFrame.minY < connectButtonFrame.maxY {
-            self.entryViewYConstraint.constant = -keyboardFrame.size.height / 2.0
+            self.entryViewYConstraint.constant = -(keyboardFrame.size.height - self.headerView.frame.height) / 2.0
         } else {
             self.entryViewYConstraint.constant = 0
         }
