Fix: Reconnection wasn't working since moving to multiple streamers

diff --git a/Engine/Extras/VirtualProduction/LiveLinkVCAM/vcam/VideoViewController+StreamingConnectionDelegate.swift b/Engine/Extras/VirtualProduction/LiveLinkVCAM/vcam/VideoViewController+StreamingConnectionDelegate.swift
--- a/Engine/Extras/VirtualProduction/LiveLinkVCAM/vcam/VideoViewController+StreamingConnectionDelegate.swift
+++ b/Engine/Extras/VirtualProduction/LiveLinkVCAM/vcam/VideoViewController+StreamingConnectionDelegate.swift
@@ -54,6 +54,24 @@ extension VideoViewController : StreamingConnectionDelegate {
     }
     
     func streamingConnection(_ connection: StreamingConnection, requestStreamerSelectionWithStreamers streamers: Array<String>, handler: @escaping (String) -> Void) {
+        self.pickerData = streamers;
+        self.selectedStreamer = streamers[0];
+        DispatchQueue.main.async {
+            self.hideConnectingAlertView() {
+                let alert = UIAlertController(title: "Select Streamer", message: "\n\n\n\n\n\n", preferredStyle: .alert)
+                
+                let picker = UIPickerView(frame: CGRect(x: 5, y: 20, width: 250, height: 140))
+                picker.dataSource = self;
+                picker.delegate = self;
+                alert.view.addSubview(picker)
+
+                alert.addAction(UIAlertAction(title: "OK", style: .default) {_ in
+                    handler(self.selectedStreamer)
+                })
+                
+                self.present(alert, animated:true)
+            }
+        }
     }
 }
 
diff --git a/Engine/Extras/VirtualProduction/LiveLinkVCAM/vcam/VideoViewController.swift b/Engine/Extras/VirtualProduction/LiveLinkVCAM/vcam/VideoViewController.swift
--- a/Engine/Extras/VirtualProduction/LiveLinkVCAM/vcam/VideoViewController.swift
+++ b/Engine/Extras/VirtualProduction/LiveLinkVCAM/vcam/VideoViewController.swift
@@ -33,6 +33,8 @@ class VideoViewController : BaseViewController {
     var headerPullDownGestureRecognizer : UIScreenEdgePanGestureRecognizer!
     
     var streamingConnection : StreamingConnection?
+    var pickerData: [String] = [String]()
+    var selectedStreamer: String = ""
 
     weak var liveLink : LiveLinkProvider?
     var dismissOnDisconnect = false
@@ -401,3 +403,22 @@ extension VideoViewController : HeaderViewDelegate {
 }
 
 
+extension VideoViewController: UIPickerViewDelegate, UIPickerViewDataSource {
+    func numberOfComponents(in pickerView: UIPickerView) -> Int {
+        // Number of columns
+        return 1
+    }
+    
+    func pickerView(_ pickerView: UIPickerView, numberOfRowsInComponent component: Int) -> Int {
+        // Number of rows
+        return pickerData.count
+    }
+    
+    func pickerView(_ pickerView: UIPickerView, titleForRow row: Int, forComponent component: Int) -> String? {
+        return pickerData[row]
+    }
+    
+    func pickerView(_ pickerView: UIPickerView, didSelectRow row: Int, inComponent component: Int) {
+        selectedStreamer = pickerData[row]
+    }
+}
diff --git a/Engine/Extras/VirtualProduction/LiveLinkVCAM/vcam/WebRTC/Services/SignalingClient.swift b/Engine/Extras/VirtualProduction/LiveLinkVCAM/vcam/WebRTC/Services/SignalingClient.swift
--- a/Engine/Extras/VirtualProduction/LiveLinkVCAM/vcam/WebRTC/Services/SignalingClient.swift
+++ b/Engine/Extras/VirtualProduction/LiveLinkVCAM/vcam/WebRTC/Services/SignalingClient.swift
@@ -34,7 +34,6 @@ final class SignalingClient {
     func connect() {
         self.webSocket.delegate = self
         self.webSocket.connect()
-        self.sendRequestStreamerList()
     }
     
     func close() {
diff --git a/Engine/Extras/VirtualProduction/LiveLinkVCAM/vcam/WebRTCStreamingConnection.swift b/Engine/Extras/VirtualProduction/LiveLinkVCAM/vcam/WebRTCStreamingConnection.swift
--- a/Engine/Extras/VirtualProduction/LiveLinkVCAM/vcam/WebRTCStreamingConnection.swift
+++ b/Engine/Extras/VirtualProduction/LiveLinkVCAM/vcam/WebRTCStreamingConnection.swift
@@ -364,6 +364,7 @@ extension WebRTCStreamingConnection: SignalClientDelegate {
     func signalClientDidConnect(_ signalClient: SignalingClient) {
         self.signalingConnected = true
         Log.info("Connected to signaling server")
+        signalClient.sendRequestStreamerList();
     }
     
     func signalClientDidDisconnect(_ signalClient: SignalingClient, error: Error?) {
