Fixed incorrect format for touch packets : this bug only appeared on devices that supported force touch (iPhone 6s -> XS)

diff --git a/Engine/Extras/VirtualProduction/LiveLinkVCAM/vcam/OSC/OSCUtility.swift b/Engine/Extras/VirtualProduction/LiveLinkVCAM/vcam/OSC/OSCUtility.swift
--- a/Engine/Extras/VirtualProduction/LiveLinkVCAM/vcam/OSC/OSCUtility.swift
+++ b/Engine/Extras/VirtualProduction/LiveLinkVCAM/vcam/OSC/OSCUtility.swift
@@ -76,6 +76,7 @@ class OSCUtility {
         data.append(OSCPacket.encode(Float(point.x), bigEndian: false))
         data.append(OSCPacket.encode(Float(point.y), bigEndian: false))
         data.append(OSCPacket.encode(Int32(finger), bigEndian: false))
+        data.append(OSCPacket.encode(Int32(0), bigEndian: false)) // user ID -- always zero.
         data.append(OSCPacket.encode(Float(force), bigEndian: false))
 
         return data;
