Fix: Gamepad ID's would clash when multiple ipads were connected

Gamepad id would always increment from 0 per device.  PixelStreaming now keeps track of all of the connected gamepads and will assign a new id for every controller that connect

diff --git a/Engine/Extras/VirtualProduction/LiveLinkVCAM/vcam/Gamepad.swift b/Engine/Extras/VirtualProduction/LiveLinkVCAM/vcam/Gamepad.swift
new file mode 100644
--- /dev/null
+++ b/Engine/Extras/VirtualProduction/LiveLinkVCAM/vcam/Gamepad.swift
@@ -0,0 +1,16 @@
+//
+//  Gamepad.swift
+//  vcam
+//
+//  Created by William Belcher on 16/3/2023.
+//  Copyright Epic Games, Inc. All Rights Reserved.
+//
+
+import Foundation
+import GameController
+
+
+class Gamepad : GCController {
+    var controller : GCController? = nil
+    var id: UInt8? = nil
+}
diff --git a/Engine/Extras/VirtualProduction/LiveLinkVCAM/vcam/StartViewController+StreamingConnectionDelegate.swift b/Engine/Extras/VirtualProduction/LiveLinkVCAM/vcam/StartViewController+StreamingConnectionDelegate.swift
--- a/Engine/Extras/VirtualProduction/LiveLinkVCAM/vcam/StartViewController+StreamingConnectionDelegate.swift
+++ b/Engine/Extras/VirtualProduction/LiveLinkVCAM/vcam/StartViewController+StreamingConnectionDelegate.swift
@@ -69,4 +69,6 @@ extension StartViewController : StreamingConnectionDelegate {
             }
         }
     }
+    func streamingConnection(_ connection: StreamingConnection, receivedGamepadResponse: UInt8) {
+    }
 }
diff --git a/Engine/Extras/VirtualProduction/LiveLinkVCAM/vcam/StreamingConnection.swift b/Engine/Extras/VirtualProduction/LiveLinkVCAM/vcam/StreamingConnection.swift
--- a/Engine/Extras/VirtualProduction/LiveLinkVCAM/vcam/StreamingConnection.swift
+++ b/Engine/Extras/VirtualProduction/LiveLinkVCAM/vcam/StreamingConnection.swift
@@ -15,6 +15,7 @@ protocol StreamingConnectionDelegate : AnyObject {
     func streamingConnection(_ connection : StreamingConnection, didDisconnectWithError err: Error?)
     func streamingConnection(_ connection : StreamingConnection, requestsTextEditWithContents contents : String, handler : @escaping (Bool, String?) -> Void)
     func streamingConnection(_ connection: StreamingConnection, requestStreamerSelectionWithStreamers streamers: Array<String>, handler: @escaping (String) -> Void)
+    func streamingConnection(_ connection: StreamingConnection, receivedGamepadResponse: UInt8)
 }
 
 enum StreamingConnectionType : String {
@@ -136,6 +137,10 @@ class StreamingConnection : NSObject {
         assertionFailure("not implemented")
     }
     
+    func sendControllerConnected() {
+        assertionFailure("not implemented")
+    }
+    
     func sendControllerAnalog(_ type :StreamingConnectionControllerInputType, controllerIndex : UInt8, value : Float) {
         assertionFailure("not implemented")
     }
@@ -147,6 +152,10 @@ class StreamingConnection : NSObject {
     func sendControllerButtonReleased(_ type : StreamingConnectionControllerInputType, controllerIndex : UInt8) {
         assertionFailure("not implemented")
     }
+    
+    func sendControllerDisconnected(controllerIndex: UInt8) {
+        assertionFailure("not implemented")
+    }
 
     func setRenderView(_ view : UIView) {
         assertionFailure("not implemented")
diff --git a/Engine/Extras/VirtualProduction/LiveLinkVCAM/vcam/VideoViewController+StreamingConnectionDelegate.swift b/Engine/Extras/VirtualProduction/LiveLinkVCAM/vcam/VideoViewController+StreamingConnectionDelegate.swift
--- a/Engine/Extras/VirtualProduction/LiveLinkVCAM/vcam/VideoViewController+StreamingConnectionDelegate.swift
+++ b/Engine/Extras/VirtualProduction/LiveLinkVCAM/vcam/VideoViewController+StreamingConnectionDelegate.swift
@@ -73,5 +73,9 @@ extension VideoViewController : StreamingConnectionDelegate {
             }
         }
     }
+    
+    func streamingConnection(_ connection: StreamingConnection, receivedGamepadResponse: UInt8) {
+        self.controllerResponseReceived(controllerIndex: receivedGamepadResponse)
+    }
 }
 
diff --git a/Engine/Extras/VirtualProduction/LiveLinkVCAM/vcam/VideoViewController.swift b/Engine/Extras/VirtualProduction/LiveLinkVCAM/vcam/VideoViewController.swift
--- a/Engine/Extras/VirtualProduction/LiveLinkVCAM/vcam/VideoViewController.swift
+++ b/Engine/Extras/VirtualProduction/LiveLinkVCAM/vcam/VideoViewController.swift
@@ -66,6 +66,8 @@ class VideoViewController : BaseViewController {
         }
     }
     
+    private var gamepads: [GCControllerPlayerIndex : Gamepad] = [:];
+    
     @IBOutlet weak var arView : ARSCNView!
 
     @IBOutlet weak var reconnectingBlurView : UIVisualEffectView!
@@ -139,11 +141,33 @@ class VideoViewController : BaseViewController {
     override func viewDidAppear(_ animated: Bool) {
         super.viewDidAppear(animated)
         UIApplication.shared.isIdleTimerDisabled = true
+        
+        // Check for already connected controllers
+        for gc in GCController.controllers() {
+            self.controllerConnected(gamepad: gc)
+        }
+
+        let notificationCenter = NotificationCenter.default;
+        
+        // Add notifications for when new controllers connect or disconnect
+        notificationCenter.addObserver(forName: .GCControllerDidConnect, object: nil, queue: .main) { (note) in
+            guard let gc = note.object as? GCController else { return }
+            self.controllerConnected(gamepad: gc)
+        }
+        
+        notificationCenter.addObserver(forName: .GCControllerDidDisconnect, object: nil, queue: .main) { (note) in
+            guard let gc = note.object as? GCController else { return }
+            self.controllerDisconnected(gamepad: gc)
+        }
     }
     
     override func viewWillDisappear(_ animated: Bool) {
         super.viewWillDisappear(animated)
         UIApplication.shared.isIdleTimerDisabled = false
+        
+        let notificationCenter = NotificationCenter.default
+        notificationCenter.removeObserver(self, name: .GCControllerDidConnect, object: nil)
+        notificationCenter.removeObserver(self, name: .GCControllerDidDisconnect, object: nil)
     }
     
     func showReconnecting(_ visible : Bool, animated: Bool) {
@@ -190,6 +214,11 @@ class VideoViewController : BaseViewController {
             Log.info("Forcing disconnect due to timeout.")
             self.forceDisconnectAndDismiss()
         })
+        // Tell UE that the controllers connected to this device are no longer used
+        for gc in GCController.controllers() {
+            self.controllerDisconnected(gamepad: gc)
+        }
+        
         self.streamingConnection?.disconnect()
     }
     
@@ -306,8 +335,10 @@ class VideoViewController : BaseViewController {
         let snapshot = gc.capture()
         guard let gp = snapshot.extendedGamepad else { return }
 
-        let controllerIndex = gc.playerIndex.rawValue
-
+        guard let controller = gamepads[gc.playerIndex] else { return }
+        // Fallback to default funcationality if we haven't received an ID from UE. This could happen if using latest app with UE < 5.2
+        let controllerIndex = controller.id ?? UInt8(gc.playerIndex.rawValue)
+        
         // dict of type -> oldValue, newValue
         let inputs : [StreamingConnectionControllerInputType : ( oldValue: Float?, newValue : Float)]  = [
             .thumbstickLeftX : ( self.gameControllerSnapshot?.extendedGamepad?.leftThumbstick.xAxis.value,  gp.leftThumbstick.xAxis.value),
@@ -333,7 +364,10 @@ class VideoViewController : BaseViewController {
         guard let gc = gameController else { return }
         guard let gp = gc.extendedGamepad else { return }
 
-        let controllerIndex = gc.playerIndex.rawValue
+        guard let controller = gamepads[gc.playerIndex] else { return }
+        // Fallback to default functionality if we haven't received an ID from UE. This could happen if using latest app with UE < 5.2
+        let controllerIndex = controller.id ?? UInt8(gc.playerIndex.rawValue)
+        
         let isRepeat = false
         
         var inputType : StreamingConnectionControllerInputType!
@@ -382,6 +416,31 @@ class VideoViewController : BaseViewController {
             sc.sendControllerButtonReleased(inputType!, controllerIndex: UInt8(controllerIndex))
         }
     }
+    
+    func controllerConnected(gamepad: GCController) {
+        guard let sc = streamingConnection else { return }
+        sc.sendControllerConnected();
+        
+        let newGamepad = Gamepad()
+        newGamepad.controller = gamepad
+        gamepads[gamepad.playerIndex] = newGamepad
+    }
+    
+    func controllerDisconnected(gamepad: GCController) {
+        guard let sc = streamingConnection else { return }
+        sc.sendControllerDisconnected(controllerIndex: gamepads[gamepad.playerIndex]!.id!);
+        
+        gamepads.removeValue(forKey: gamepad.playerIndex)
+    }
+    
+    func controllerResponseReceived(controllerIndex: UInt8) {
+        for gamepad in gamepads.values {
+            if(gamepad.id == nil) {
+                gamepad.id = controllerIndex;
+                break;
+            }
+        }
+    }
 }
 
 extension VideoViewController : HeaderViewDelegate {
diff --git a/Engine/Extras/VirtualProduction/LiveLinkVCAM/vcam/WebRTC/Extensions/PixelStreamingToClientMessage.swift b/Engine/Extras/VirtualProduction/LiveLinkVCAM/vcam/WebRTC/Extensions/PixelStreamingToClientMessage.swift
--- a/Engine/Extras/VirtualProduction/LiveLinkVCAM/vcam/WebRTC/Extensions/PixelStreamingToClientMessage.swift
+++ b/Engine/Extras/VirtualProduction/LiveLinkVCAM/vcam/WebRTC/Extensions/PixelStreamingToClientMessage.swift
@@ -21,6 +21,7 @@ enum PixelStreamingToClientMessage: UInt8, Codable {
     case FileExtension = 8
     case FileMimeType = 9
     case FileContents = 10
+    case GamepadResponse = 13
 }
 
 
@@ -33,3 +34,7 @@ struct PixelStreamingToClientShowOnScreenKeyboardCommand: Decodable {
     let showOnScreenKeyboard: Bool
     let contents: String
 }
+
+struct PixelStreamingToClientGamepadResponse: Decodable {
+    let controllerId: UInt8
+}
diff --git a/Engine/Extras/VirtualProduction/LiveLinkVCAM/vcam/WebRTC/Extensions/PixelStreamingToStreamerMessage.swift b/Engine/Extras/VirtualProduction/LiveLinkVCAM/vcam/WebRTC/Extensions/PixelStreamingToStreamerMessage.swift
--- a/Engine/Extras/VirtualProduction/LiveLinkVCAM/vcam/WebRTC/Extensions/PixelStreamingToStreamerMessage.swift
+++ b/Engine/Extras/VirtualProduction/LiveLinkVCAM/vcam/WebRTC/Extensions/PixelStreamingToStreamerMessage.swift
@@ -39,6 +39,8 @@ enum PixelStreamingToStreamerMessage : UInt8, Codable {
     case GamepadButtonPressed = 90
     case GamepadButtonReleased = 91
     case GamepadAnalog = 92
+    case GamepadConnected = 93
+    case GamepadDisconnected = 94
     
     case Transform = 100
     case TextboxEntry = 52
diff --git a/Engine/Extras/VirtualProduction/LiveLinkVCAM/vcam/WebRTCStreamingConnection.swift b/Engine/Extras/VirtualProduction/LiveLinkVCAM/vcam/WebRTCStreamingConnection.swift
--- a/Engine/Extras/VirtualProduction/LiveLinkVCAM/vcam/WebRTCStreamingConnection.swift
+++ b/Engine/Extras/VirtualProduction/LiveLinkVCAM/vcam/WebRTCStreamingConnection.swift
@@ -253,6 +253,17 @@ class WebRTCStreamingConnection : StreamingConnection {
         }
     }
     
+    override func sendControllerConnected() {
+        guard let client = webRTCClient else { return }
+        
+        var bytes: [UInt8] = []
+
+        // Write message type using 1 byte
+        bytes.append(PixelStreamingToStreamerMessage.GamepadConnected.rawValue)
+        
+        client.sendData(Data(bytes))
+    }
+    
     override func sendControllerAnalog(_ type : StreamingConnectionControllerInputType, controllerIndex : UInt8, value : Float) {
         
         guard let client = webRTCClient else { return }
@@ -310,6 +321,18 @@ class WebRTCStreamingConnection : StreamingConnection {
        
        client.sendData(Data(bytes))
     }
+    
+    override func sendControllerDisconnected(controllerIndex: UInt8) {
+        guard let client = webRTCClient else { return }
+        
+        var bytes: [UInt8] = []
+
+        // Write message type using 1 byte
+        bytes.append(PixelStreamingToStreamerMessage.GamepadDisconnected.rawValue)
+        bytes.append(controllerIndex)
+        
+        client.sendData(Data(bytes))
+    }
 
     
     func attachVideoTrack() {
@@ -525,6 +548,17 @@ extension WebRTCStreamingConnection: WebRTCClientDelegate {
                             Log.error("An error occurred parsing the command `\(command ?? "<invalid>")` : \(error.localizedDescription)")
                         }
                     }
+                case .GamepadResponse:
+                    let response: String? = String(data: data.dropFirst(), encoding: .utf16LittleEndian)
+                    Log.info("response = \(response ?? "NULL")")
+                    if let responseData = response?.data(using: .utf8) {
+                        do {
+                            let responseJson: PixelStreamingToClientGamepadResponse = try JSONDecoder().decode(PixelStreamingToClientGamepadResponse.self, from: responseData)
+                            self.delegate?.streamingConnection(self, receivedGamepadResponse: responseJson.controllerId)
+                        } catch {
+                            Log.error("An error occurred parsing the response `\(response ?? "<invalid>")` : \(error.localizedDescription)")
+                        }
+                    }
                 case .QualityControlOwnership:
                     fallthrough
                 case .Response:
