[EpicStageApp] Refactor settings system to break it up into components with persisted values

diff --git a/Engine/Extras/VirtualProduction/EpicStageApp/pubspec.yaml b/Engine/Extras/VirtualProduction/EpicStageApp/pubspec.yaml
--- a/Engine/Extras/VirtualProduction/EpicStageApp/pubspec.yaml
+++ b/Engine/Extras/VirtualProduction/EpicStageApp/pubspec.yaml
@@ -40,7 +40,7 @@ dependencies:
   web_socket_channel: ^2.1.0
   provider: ^6.0.2
   flutter_svg: ^1.0.2
-  shared_preferences: ^2.0.15
+  streaming_shared_preferences: ^2.0.0
   webview_flutter: ^3.0.4
   logging: ^1.1.0
   intl: ^0.17.0
