Fix: device resolution sometimes wouldn't send

diff --git a/Engine/Extras/VirtualProduction/LiveLinkVCAM/vcam/WebRTC/Services/WebRTCClient.swift b/Engine/Extras/VirtualProduction/LiveLinkVCAM/vcam/WebRTC/Services/WebRTCClient.swift
--- a/Engine/Extras/VirtualProduction/LiveLinkVCAM/vcam/WebRTC/Services/WebRTCClient.swift
+++ b/Engine/Extras/VirtualProduction/LiveLinkVCAM/vcam/WebRTC/Services/WebRTCClient.swift
@@ -140,6 +140,34 @@ final class WebRTCClient: NSObject {
         }
     }
     
+    func sendDeviceResolution() {
+        let bounds = UIScreen.main.nativeBounds
+        
+        // the device nativeBounds are always in portrait orientation, so we swap them here because we always run
+        // in landscape.
+        let resolutionCommand = PixelStreamingToStreamerResolutionCommand(width: Int(bounds.height), height: Int(bounds.width))
+        
+        do {
+            let jsonData = try JSONEncoder().encode(resolutionCommand)
+            let jsonString = String(data: jsonData, encoding: .utf8)!
+            
+            var bytes: [UInt8] = []
+            // Append the message type
+            bytes.append(PixelStreamingToStreamerMessage.Command.rawValue)
+            bytes.append(contentsOf: UInt16(jsonString.count).toBytes())
+            // Append each character in the json string represented by its UTF16 character code
+            for char in jsonString {
+                let charCode = char.utf16[char.utf16.index(char.utf16.startIndex, offsetBy: 0)]
+                bytes.append(contentsOf: charCode.toBytes())
+            }
+            
+            Log.info("Sending remote resolution")
+            self.sendData(Data(bytes))
+        } catch {
+            print(error)
+        }
+    }
+    
 }
 
 // MARK: Data Channels
@@ -206,6 +234,7 @@ extension WebRTCClient: RTCPeerConnectionDelegate {
     
     func peerConnection(_ peerConnection: RTCPeerConnection, didOpen dataChannel: RTCDataChannel) {
         debugPrint("peerConnection did open data channel")
+        self.sendDeviceResolution()
     }
     
     func peerConnection(_ peerConnection: RTCPeerConnection, didAdd rtpReceiver: RTCRtpReceiver, streams mediaStreams: [RTCMediaStream]) {
diff --git a/Engine/Extras/VirtualProduction/LiveLinkVCAM/vcam/WebRTCStreamingConnection.swift b/Engine/Extras/VirtualProduction/LiveLinkVCAM/vcam/WebRTCStreamingConnection.swift
--- a/Engine/Extras/VirtualProduction/LiveLinkVCAM/vcam/WebRTCStreamingConnection.swift
+++ b/Engine/Extras/VirtualProduction/LiveLinkVCAM/vcam/WebRTCStreamingConnection.swift
@@ -341,36 +341,6 @@ class WebRTCStreamingConnection : StreamingConnection {
             self.keyboardControls = KeyboardControls(webRTC)
             view.attachVideoTrack(track: track)
             view.attachTouchDelegate(delegate: self.touchControls!)
-            self.sendDeviceResolution()
-        }
-    }
-    
-    func sendDeviceResolution() {
-        if let client = self.webRTCClient {
-            let bounds = UIScreen.main.nativeBounds
-            
-            // the device nativeBounds are always in portrait orientation, so we swap them here because we always run
-            // in landscape.
-            let resolutionCommand = PixelStreamingToStreamerResolutionCommand(width: Int(bounds.height), height: Int(bounds.width))
-            
-            do {
-                let jsonData = try JSONEncoder().encode(resolutionCommand)
-                let jsonString = String(data: jsonData, encoding: .utf8)!
-                
-                var bytes: [UInt8] = []
-                // Append the message type
-                bytes.append(PixelStreamingToStreamerMessage.Command.rawValue)
-                bytes.append(contentsOf: UInt16(jsonString.count).toBytes())
-                // Append each character in the json string represented by its UTF16 character code
-                for char in jsonString {
-                    let charCode = char.utf16[char.utf16.index(char.utf16.startIndex, offsetBy: 0)]
-                    bytes.append(contentsOf: charCode.toBytes())
-                }
-                
-                client.sendData(Data(bytes))
-            } catch {
-                print(error)
-            }
         }
     }
 }
