Added check for json contents to avoid exception & logging spam.

diff --git a/Engine/Extras/VirtualProduction/LiveLinkVCAM/vcam/WebRTC/Extensions/PixelStreamingToClientMessage.swift b/Engine/Extras/VirtualProduction/LiveLinkVCAM/vcam/WebRTC/Extensions/PixelStreamingToClientMessage.swift
--- a/Engine/Extras/VirtualProduction/LiveLinkVCAM/vcam/WebRTC/Extensions/PixelStreamingToClientMessage.swift
+++ b/Engine/Extras/VirtualProduction/LiveLinkVCAM/vcam/WebRTC/Extensions/PixelStreamingToClientMessage.swift
@@ -32,7 +32,7 @@ struct PixelStreamingToClientCommand: Decodable {
 
 struct PixelStreamingToClientShowOnScreenKeyboardCommand: Decodable {
     let showOnScreenKeyboard: Bool
-    let contents: String
+    let contents: String?
 }
 
 struct PixelStreamingToClientGamepadResponse: Decodable {
diff --git a/Engine/Extras/VirtualProduction/LiveLinkVCAM/vcam/WebRTCStreamingConnection.swift b/Engine/Extras/VirtualProduction/LiveLinkVCAM/vcam/WebRTCStreamingConnection.swift
--- a/Engine/Extras/VirtualProduction/LiveLinkVCAM/vcam/WebRTCStreamingConnection.swift
+++ b/Engine/Extras/VirtualProduction/LiveLinkVCAM/vcam/WebRTCStreamingConnection.swift
@@ -505,11 +505,11 @@ extension WebRTCStreamingConnection: WebRTCClientDelegate {
                             if commandJson.command == "onScreenKeyboard" {
                                 let showKeyboardJson: PixelStreamingToClientShowOnScreenKeyboardCommand = try JSONDecoder().decode(PixelStreamingToClientShowOnScreenKeyboardCommand.self, from: commandData)
 
-                                if showKeyboardJson.showOnScreenKeyboard {
-                                    self.delegate?.streamingConnection(self, requestsTextEditWithContents: showKeyboardJson.contents) { (success, newContents) in
-                                        guard let contents = newContents else { return }
+                                if showKeyboardJson.showOnScreenKeyboard, let contents = showKeyboardJson.contents {
+                                    self.delegate?.streamingConnection(self, requestsTextEditWithContents: contents) { (success, newContents) in
+                                        guard let enteredContents = newContents else { return }
                                         if success {
-                                            self.keyboardControls?.submitString(contents)
+                                            self.keyboardControls?.submitString(enteredContents)
                                         }
                                     }
                                 }
