Reconnection flow

diff --git a/Engine/Extras/VirtualProduction/LiveLinkVCAM/vcam/StartViewController+StreamingConnectionDelegate.swift b/Engine/Extras/VirtualProduction/LiveLinkVCAM/vcam/StartViewController+StreamingConnectionDelegate.swift
--- a/Engine/Extras/VirtualProduction/LiveLinkVCAM/vcam/StartViewController+StreamingConnectionDelegate.swift
+++ b/Engine/Extras/VirtualProduction/LiveLinkVCAM/vcam/StartViewController+StreamingConnectionDelegate.swift
@@ -46,6 +46,9 @@ extension StartViewController : StreamingConnectionDelegate {
         }
     }
     
+    func streamingConnection(_ connection: StreamingConnection, exitWithErr err: Error?) {
+    }
+    
     func streamingConnection(_ connection: StreamingConnection, requestsTextEditWithContents contents: String, handler: @escaping (Bool, String?) -> Void) {
     }
     
diff --git a/Engine/Extras/VirtualProduction/LiveLinkVCAM/vcam/StreamingConnection.swift b/Engine/Extras/VirtualProduction/LiveLinkVCAM/vcam/StreamingConnection.swift
--- a/Engine/Extras/VirtualProduction/LiveLinkVCAM/vcam/StreamingConnection.swift
+++ b/Engine/Extras/VirtualProduction/LiveLinkVCAM/vcam/StreamingConnection.swift
@@ -13,6 +13,7 @@ protocol StreamingConnectionDelegate : AnyObject {
     
     func streamingConnectionDidConnect(_ connection : StreamingConnection)
     func streamingConnection(_ connection : StreamingConnection, didDisconnectWithError err: Error?)
+    func streamingConnection(_ connection : StreamingConnection, exitWithErr err: Error?)
     func streamingConnection(_ connection : StreamingConnection, requestsTextEditWithContents contents : String, handler : @escaping (Bool, String?) -> Void)
     func streamingConnection(_ connection: StreamingConnection, requestStreamerSelectionWithStreamers streamers: Array<String>, handler: @escaping (String) -> Void)
     func streamingConnection(_ connection: StreamingConnection, receivedGamepadResponse: UInt8)
diff --git a/Engine/Extras/VirtualProduction/LiveLinkVCAM/vcam/VideoViewController+StreamingConnectionDelegate.swift b/Engine/Extras/VirtualProduction/LiveLinkVCAM/vcam/VideoViewController+StreamingConnectionDelegate.swift
--- a/Engine/Extras/VirtualProduction/LiveLinkVCAM/vcam/VideoViewController+StreamingConnectionDelegate.swift
+++ b/Engine/Extras/VirtualProduction/LiveLinkVCAM/vcam/VideoViewController+StreamingConnectionDelegate.swift
@@ -33,7 +33,19 @@ extension VideoViewController : StreamingConnectionDelegate {
                 self.reconnect()
             }
         }
+    }
+    
+    func streamingConnection(_ connection: StreamingConnection, exitWithErr err: Error?) {
+        DispatchQueue.main.async {
+            self.presentingViewController?.dismiss(animated: true, completion: nil)
+            
+            let alert = UIAlertController(title: nil, message: err?.localizedDescription, preferredStyle: .alert)
+            alert.addAction(UIAlertAction(title: Localized.buttonOK(), style: .default) {_ in
+                self.forceDisconnectAndDismiss()
+            })
 
+            self.present(alert, animated: true)
+        }
     }
     
     func streamingConnection(_ connection: StreamingConnection, requestsTextEditWithContents contents: String, handler: @escaping (Bool, String?) -> Void) {
diff --git a/Engine/Extras/VirtualProduction/LiveLinkVCAM/vcam/WebRTC/Services/SignalingClient.swift b/Engine/Extras/VirtualProduction/LiveLinkVCAM/vcam/WebRTC/Services/SignalingClient.swift
--- a/Engine/Extras/VirtualProduction/LiveLinkVCAM/vcam/WebRTC/Services/SignalingClient.swift
+++ b/Engine/Extras/VirtualProduction/LiveLinkVCAM/vcam/WebRTC/Services/SignalingClient.swift
@@ -25,6 +25,8 @@ final class SignalingClient {
     private let encoder = JSONEncoder()
     private let webSocket: WebSocketProvider
     private var shouldReconnect : Bool = true
+    var isReconnecting  : Bool = false
+    var reconnectAttempt : Int = 0
     weak var delegate: SignalClientDelegate?
     
     init(webSocket: WebSocketProvider) {
@@ -42,6 +44,12 @@ final class SignalingClient {
         self.webSocket.delegate = nil
     }
     
+    func reconnect() {
+        self.shouldReconnect = true
+        self.webSocket.delegate = self
+        self.webSocket.close()
+    }
+    
     func send(sdp rtcSdp: RTCSessionDescription) {
         let message = Message.sdp(SessionDescription(from: rtcSdp))
         do {
@@ -101,6 +109,8 @@ extension SignalingClient: WebSocketProviderDelegate {
         DispatchQueue.global().asyncAfter(deadline: .now() + 2) {
             if self.shouldReconnect {
                 debugPrint("Trying to reconnect to signaling server...")
+                self.isReconnecting = true
+                self.reconnectAttempt += 1
                 self.webSocket.connect()
             }
         }
@@ -141,7 +151,7 @@ extension SignalingClient: WebSocketProviderDelegate {
         case .requestStreamerList:
             break
         case .subscribe(let streamerId):
-            break;
+            break
         }
     }
     
diff --git a/Engine/Extras/VirtualProduction/LiveLinkVCAM/vcam/WebRTCStreamingConnection.swift b/Engine/Extras/VirtualProduction/LiveLinkVCAM/vcam/WebRTCStreamingConnection.swift
--- a/Engine/Extras/VirtualProduction/LiveLinkVCAM/vcam/WebRTCStreamingConnection.swift
+++ b/Engine/Extras/VirtualProduction/LiveLinkVCAM/vcam/WebRTCStreamingConnection.swift
@@ -19,7 +19,7 @@ class WebRTCStreamingConnection : StreamingConnection {
     private var keyboardControls: KeyboardControls?
     private var webRTCView : WebRTCView?
     private var rtcVideoTrack : RTCVideoTrack?
-    
+    private var subscribedStream : String = ""
     private var signalingConnected = false
     private var hasRemoteSdp = false
     private var hasLocalSdp = false
@@ -426,13 +426,40 @@ extension WebRTCStreamingConnection: SignalClientDelegate {
     }
     
     func signalClient(_ signalClient: SignalingClient, didReceiveStreamerList streamerList: Array<String>) {
-        // If we only have a single streamer, no need to show the selection dialogue
-        if streamerList.count == 1 {
-            signalClient.subscribe(streamerList[0])
-        } else if streamerList.count > 1 {
-            // Otherwise make sure we have more than 1 and display the picker
-            self.delegate?.streamingConnection(self, requestStreamerSelectionWithStreamers: streamerList) { (selectedStreamer) in
-                signalClient.subscribe(selectedStreamer)
+        if signalClient.isReconnecting
+        {
+            if streamerList.contains(subscribedStream)
+            {
+                // The stream we were originally subscribed to has come back, so reconnect to it
+                signalClient.isReconnecting = false;
+                signalClient.reconnectAttempt = 0;
+                signalClient.subscribe(subscribedStream)
+            }
+            else if(signalClient.reconnectAttempt < 3)
+            {
+                // The stream's still not back up, try reconnecting a few more times
+                signalClient.reconnect()
+            }
+            else
+            {
+                // We've tried a few times and still no stream, return to main menu
+                Log.info("Unable to reconnect to previously connected VCam")
+                self.delegate?.streamingConnection(self, exitWithErr: NSError(domain: "", code: 1, userInfo: [NSLocalizedDescriptionKey : "Unable to reconnect to VCam after 3 tries" ] ))
+            }
+            
+        }
+        else
+        {
+            // If we only have a single streamer, no need to show the selection dialogue
+            if streamerList.count == 1 {
+                self.subscribedStream = streamerList[0]
+                signalClient.subscribe(streamerList[0])
+            } else if streamerList.count > 1 {
+                // Otherwise make sure we have more than 1 and display the picker
+                self.delegate?.streamingConnection(self, requestStreamerSelectionWithStreamers: streamerList) { (selectedStreamer) in
+                    self.subscribedStream = selectedStreamer
+                    signalClient.subscribe(selectedStreamer)
+                }
             }
         }
     }
