Live Link VCAM iOS v1.3.1 (build 12) : Localization for Simplified Chinese,  Japanese,  and Korean for various multi-streamer operations and error messages.

diff --git a/Engine/Extras/VirtualProduction/LiveLinkVCAM/vcam/Localized.swift b/Engine/Extras/VirtualProduction/LiveLinkVCAM/vcam/Localized.swift
--- a/Engine/Extras/VirtualProduction/LiveLinkVCAM/vcam/Localized.swift
+++ b/Engine/Extras/VirtualProduction/LiveLinkVCAM/vcam/Localized.swift
@@ -68,4 +68,10 @@ class Localized {
     class func subjectName() -> String {
         NSLocalizedString("livelink-subjectname", value: "Subject Name", comment: "The name of the subject (a person, device, character) of the capture")
     }
+    class func messageCouldntConnect() -> String {
+        NSLocalizedString("message-couldntconnect", value: "Couldn't connect", comment: "This text is followed by a colon and an error string")
+    }
+    class func titleSelectStream() -> String {
+        NSLocalizedString("title-selectstream", value: "Select Stream", comment: "Selection of a video stream")
+    }
 }
diff --git a/Engine/Extras/VirtualProduction/LiveLinkVCAM/vcam/StartViewController+StreamingConnectionDelegate.swift b/Engine/Extras/VirtualProduction/LiveLinkVCAM/vcam/StartViewController+StreamingConnectionDelegate.swift
--- a/Engine/Extras/VirtualProduction/LiveLinkVCAM/vcam/StartViewController+StreamingConnectionDelegate.swift
+++ b/Engine/Extras/VirtualProduction/LiveLinkVCAM/vcam/StartViewController+StreamingConnectionDelegate.swift
@@ -33,8 +33,8 @@ extension StartViewController : StreamingConnectionDelegate {
                 if let e = err {
                     Log.info("StreamingConnection \(connection.name) disconnected with error : \(e.localizedDescription)")
 
-                    let errorAlert = UIAlertController(title: "Error", message: "Couldn't connect : \(e.localizedDescription)", preferredStyle: .alert)
-                    errorAlert.addAction(UIAlertAction(title: "OK", style: .default, handler: { _ in
+                    let errorAlert = UIAlertController(title: Localized.titleError(), message: "\(Localized.messageCouldntConnect()) : \(e.localizedDescription)", preferredStyle: .alert)
+                    errorAlert.addAction(UIAlertAction(title: Localized.buttonOK(), style: .default, handler: { _ in
                         self.hideConnectingView() {}
                     }))
                     self.present(errorAlert, animated:true)
@@ -53,8 +53,8 @@ extension StartViewController : StreamingConnectionDelegate {
                 if let e = err {
                     Log.info("StreamingConnection \(connection.name) disconnected with error : \(e.localizedDescription)")
 
-                    let errorAlert = UIAlertController(title: "Error", message: "Couldn't connect : \(e.localizedDescription)", preferredStyle: .alert)
-                    errorAlert.addAction(UIAlertAction(title: "OK", style: .default, handler: { _ in
+                    let errorAlert = UIAlertController(title: Localized.titleError(), message: "\(Localized.messageCouldntConnect()) : \(e.localizedDescription)", preferredStyle: .alert)
+                    errorAlert.addAction(UIAlertAction(title: Localized.buttonOK(), style: .default, handler: { _ in
                         self.hideConnectingView() {}
                     }))
                     self.present(errorAlert, animated:true)
@@ -74,14 +74,14 @@ extension StartViewController : StreamingConnectionDelegate {
         self.selectedStreamer = streamers[0];
         DispatchQueue.main.async {
             self.hideConnectingView() {
-                let alert = UIAlertController(title: "Select Streamer", message: "\n\n\n\n\n\n", preferredStyle: .alert)
+                let alert = UIAlertController(title: Localized.titleSelectStream(), message: "\n\n\n\n\n\n", preferredStyle: .alert)
                 
                 let picker = UIPickerView(frame: CGRect(x: 5, y: 20, width: 250, height: 140))
                 picker.dataSource = self;
                 picker.delegate = self;
                 alert.view.addSubview(picker)
 
-                alert.addAction(UIAlertAction(title: "OK", style: .default) {_ in
+                alert.addAction(UIAlertAction(title: Localized.buttonOK(), style: .default) {_ in
                     handler(self.selectedStreamer)
                 })
                 
diff --git a/Engine/Extras/VirtualProduction/LiveLinkVCAM/vcam/StartViewController.swift b/Engine/Extras/VirtualProduction/LiveLinkVCAM/vcam/StartViewController.swift
--- a/Engine/Extras/VirtualProduction/LiveLinkVCAM/vcam/StartViewController.swift
+++ b/Engine/Extras/VirtualProduction/LiveLinkVCAM/vcam/StartViewController.swift
@@ -306,7 +306,7 @@ class StartViewController : BaseViewController {
 
             } catch {
                 
-                showConnectionErrorAlert("Couldn't connect : \(error.localizedDescription)")
+                showConnectionErrorAlert("\(Localized.messageCouldntConnect()) : \(error.localizedDescription)")
             }
         }
     }
@@ -314,7 +314,7 @@ class StartViewController : BaseViewController {
     func showConnectionErrorAlert(_ message : String) {
         hideConnectingAlertView() {
 
-            let errorAlert = UIAlertController(title: Localized.titleError(), message: "Couldn't connect : \(message)", preferredStyle: .alert)
+            let errorAlert = UIAlertController(title: Localized.titleError(), message: "\(Localized.messageCouldntConnect()) : \(message)", preferredStyle: .alert)
             errorAlert.addAction(UIAlertAction(title: Localized.buttonOK(), style: .default, handler: { _ in
                 self.hideConnectingView { }
             }))
diff --git a/Engine/Extras/VirtualProduction/LiveLinkVCAM/vcam/VideoViewController+StreamingConnectionDelegate.swift b/Engine/Extras/VirtualProduction/LiveLinkVCAM/vcam/VideoViewController+StreamingConnectionDelegate.swift
--- a/Engine/Extras/VirtualProduction/LiveLinkVCAM/vcam/VideoViewController+StreamingConnectionDelegate.swift
+++ b/Engine/Extras/VirtualProduction/LiveLinkVCAM/vcam/VideoViewController+StreamingConnectionDelegate.swift
@@ -51,7 +51,10 @@ extension VideoViewController : StreamingConnectionDelegate {
     func streamingConnection(_ connection: StreamingConnection, requestsTextEditWithContents contents: String, handler: @escaping (Bool, String?) -> Void) {
         
         DispatchQueue.main.async {
-            let alert = UIAlertController(title: nil, message: contents.count == 0 ? "Enter the value:" : "Update the value:", preferredStyle: .alert)
+            let alert = UIAlertController(title: nil, message: contents.count == 0 ?
+                                          NSLocalizedString("message-entervalue", value: "Enter the value", comment: "Entering value in a textbox") :
+                                          NSLocalizedString("message-updatevalue", value: "Update the value", comment: "Updating value in a textbox"),
+                                          preferredStyle: .alert)
             alert.addTextField()
             alert.textFields![0].text = contents
             alert.addAction(UIAlertAction(title: Localized.buttonOK(), style: .default) {_ in
@@ -70,14 +73,14 @@ extension VideoViewController : StreamingConnectionDelegate {
         self.selectedStreamer = streamers[0];
         DispatchQueue.main.async {
             self.hideConnectingAlertView() {
-                let alert = UIAlertController(title: "Select Streamer", message: "\n\n\n\n\n\n", preferredStyle: .alert)
+                let alert = UIAlertController(title: Localized.titleSelectStream(), message: "\n\n\n\n\n\n", preferredStyle: .alert)
                 
                 let picker = UIPickerView(frame: CGRect(x: 5, y: 20, width: 250, height: 140))
                 picker.dataSource = self;
                 picker.delegate = self;
                 alert.view.addSubview(picker)
 
-                alert.addAction(UIAlertAction(title: "OK", style: .default) {_ in
+                alert.addAction(UIAlertAction(title: Localized.buttonOK(), style: .default) {_ in
                     handler(self.selectedStreamer)
                 })
                 
diff --git a/Engine/Extras/VirtualProduction/LiveLinkVCAM/vcam/WebRTC/Services/WebSocketProvider/NativeWebSocketProvider.swift b/Engine/Extras/VirtualProduction/LiveLinkVCAM/vcam/WebRTC/Services/WebSocketProvider/NativeWebSocketProvider.swift
--- a/Engine/Extras/VirtualProduction/LiveLinkVCAM/vcam/WebRTC/Services/WebSocketProvider/NativeWebSocketProvider.swift
+++ b/Engine/Extras/VirtualProduction/LiveLinkVCAM/vcam/WebRTC/Services/WebSocketProvider/NativeWebSocketProvider.swift
@@ -34,7 +34,7 @@ class NativeWebSocketProvider : NSObject, WebSocketProvider, URLSessionWebSocket
     }
     
     func close() {
-        let error = NSError(domain: "NativeWebSocketProvider", code: 0, userInfo: [NSLocalizedDescriptionKey : "Manually closed websocket."])
+        let error = NSError(domain: "NativeWebSocketProvider", code: 0, userInfo: [NSLocalizedDescriptionKey : NSLocalizedString("error-closedwebsocket", value: "Manually closed websocket.", comment: "Error message")])
         closeConnection(error)
     }
     
diff --git a/Engine/Extras/VirtualProduction/LiveLinkVCAM/vcam/WebRTCStreamingConnection.swift b/Engine/Extras/VirtualProduction/LiveLinkVCAM/vcam/WebRTCStreamingConnection.swift
--- a/Engine/Extras/VirtualProduction/LiveLinkVCAM/vcam/WebRTCStreamingConnection.swift
+++ b/Engine/Extras/VirtualProduction/LiveLinkVCAM/vcam/WebRTCStreamingConnection.swift
@@ -445,12 +445,16 @@ extension WebRTCStreamingConnection: SignalClientDelegate {
             } else {
                 // We've exhausted our reconnect attempts, return to main menu
                 self.reconnectAttempt = 0
-                self.delegate?.streamingConnection(self, exitWithError: NSError(domain: "", code: 1, userInfo: [NSLocalizedDescriptionKey: "Unable to reconnect to \(subscribedStreamer) after \(maxReconnectAttempts) attempts"]))
+                self.delegate?.streamingConnection(self, exitWithError:
+                                                    NSError(domain: "",
+                                                            code: 1,
+                                                            userInfo: [NSLocalizedDescriptionKey:
+                                                                        String(format: NSLocalizedString("error-unablereconnect", value:"Unable to reconnect to %s after %d attempts.", comment: "Error message : %s is replace with a string and %d with a number"), subscribedStreamer, maxReconnectAttempts)]))
             }
             
         } else {
             if streamerList.count == 0 {
-                self.delegate?.streamingConnection(self, exitWithError: NSError(domain: "", code: 1, userInfo: [NSLocalizedDescriptionKey: "No streamers connected"]))
+                self.delegate?.streamingConnection(self, exitWithError: NSError(domain: "", code: 1, userInfo: [NSLocalizedDescriptionKey: NSLocalizedString("error-nostream", value:"No stream connected.", comment: "Error message.")]))
             } else if streamerList.count == 1 {
                 // If we only have a single streamer, no need to show the selection dialogue
                 self.subscribedStreamer = streamerList[0]
@@ -508,7 +512,7 @@ extension WebRTCStreamingConnection: WebRTCClientDelegate {
         case .disconnected:
             self.delegate?.streamingConnection(self, didDisconnectWithError: nil)
         case .failed:
-            self.delegate?.streamingConnection(self, didDisconnectWithError: NSError(domain: "", code: 1, userInfo: [NSLocalizedDescriptionKey : "Failed to connect." ] ))
+            self.delegate?.streamingConnection(self, didDisconnectWithError: NSError(domain: "", code: 1, userInfo: [NSLocalizedDescriptionKey : NSLocalizedString("error-connectfailed", value:"Failed to connect.", comment: "Error message.") ] ))
         default:
             break
         }
diff --git a/Engine/Extras/VirtualProduction/LiveLinkVCAM/vcam/ja.lproj/Localizable.strings b/Engine/Extras/VirtualProduction/LiveLinkVCAM/vcam/ja.lproj/Localizable.strings
--- a/Engine/Extras/VirtualProduction/LiveLinkVCAM/vcam/ja.lproj/Localizable.strings
+++ b/Engine/Extras/VirtualProduction/LiveLinkVCAM/vcam/ja.lproj/Localizable.strings
@@ -168,3 +168,28 @@
 
 /* Pixel Streaming */
 "connectiontype-value-pixelstreaming" = "ピクセルストリーミング";
+
+/* This text is followed by an error string, for example "Couldn't connect : host not found.", where "host not
+   found." will be added by the internal exception." */
+"message-couldntconnect" = "接続不可";
+
+/* Entering value in a textbox */
+"message-entervalue" = "値を入力:";
+
+/* Updating a value in a textbox */
+"message-updatevalue" = "値を更新:";
+
+/* Selection of a video stream */
+"title-selectstream" = "ストリームを選択";
+
+/* Error message. */
+"error-closedwebsocket" = "手動で閉じたWebSocket";
+
+/* Error message : %s is replace with a string and %d with a number */
+"error-unablereconnect" = "%d 回試行しましたが、%s に再接続できません。";
+
+/* Error message. */
+"error-nostream" = "ストリームに接続されていません。";
+
+/* Error message. */
+"error-connectfailed" = "接続に失敗しました。";
diff --git a/Engine/Extras/VirtualProduction/LiveLinkVCAM/vcam/ko.lproj/Localizable.strings b/Engine/Extras/VirtualProduction/LiveLinkVCAM/vcam/ko.lproj/Localizable.strings
--- a/Engine/Extras/VirtualProduction/LiveLinkVCAM/vcam/ko.lproj/Localizable.strings
+++ b/Engine/Extras/VirtualProduction/LiveLinkVCAM/vcam/ko.lproj/Localizable.strings
@@ -168,3 +168,28 @@
 
 /* Pixel Streaming */
 "connectiontype-value-pixelstreaming" = "픽셀 스트리밍";
+
+/* This text is followed by an error string, for example "Couldn't connect : host not found.", where "host not
+   found." will be added by the internal exception." */
+"message-couldntconnect" = "연결 불가";
+
+/* Entering value in a textbox */
+"message-entervalue" = "값 입력:";
+
+/* Updating a value in a textbox */
+"message-updatevalue" = "값 업데이트:";
+
+/* Selection of a video stream */
+"title-selectstream" = "스트림 선택";
+
+/* Error message. */
+"error-closedwebsocket" = "웹소켓이 수동으로 닫혔습니다.";
+
+/* Error message : %s is replace with a string and %d with a number */
+"error-unablereconnect" = "%d번 시도한 후에는 %s에 다시 연결할 수 없습니다.";
+
+/* Error message. */
+"error-nostream" = "연결된 스트림이 없습니다.";
+
+/* Error message. */
+"error-connectfailed" = "연결하지 못했습니다.";
diff --git a/Engine/Extras/VirtualProduction/LiveLinkVCAM/vcam/zh-Hans.lproj/Localizable.strings b/Engine/Extras/VirtualProduction/LiveLinkVCAM/vcam/zh-Hans.lproj/Localizable.strings
--- a/Engine/Extras/VirtualProduction/LiveLinkVCAM/vcam/zh-Hans.lproj/Localizable.strings
+++ b/Engine/Extras/VirtualProduction/LiveLinkVCAM/vcam/zh-Hans.lproj/Localizable.strings
@@ -169,3 +169,27 @@
 /* Pixel Streaming */
 "connectiontype-value-pixelstreaming" = "像素流送";
 
+/* This text is followed by an error string, for example "Couldn't connect : host not found.", where "host not
+   found." will be added by the internal exception." */
+"message-couldntconnect" = "无法连接";
+
+/* Entering value in a textbox */
+"message-entervalue" = "输入值：";
+
+/* Updating a value in a textbox */
+"message-updatevalue" = "更新值：";
+
+/* Selection of a video stream */
+"title-selectstream" = "选择流";
+
+/* Error message. */
+"error-closedwebsocket" = "手动关闭WebSocket。";
+
+/* Error message : %s is replace with a string and %d with a number */
+"error-unablereconnect" = "经过%d次尝试，无法重新链接到%s。";
+
+/* Error message. */
+"error-nostream" = "未连接到流。";
+
+/* Error message. */
+"error-connectfailed" = "连接失败。";
