Live Link VCAM 1.3.1 build 13 : Fixed a couple of missing JA strings.

diff --git a/Engine/Extras/VirtualProduction/LiveLinkVCAM/vcam/ja.lproj/Main.strings b/Engine/Extras/VirtualProduction/LiveLinkVCAM/vcam/ja.lproj/Main.strings
--- a/Engine/Extras/VirtualProduction/LiveLinkVCAM/vcam/ja.lproj/Main.strings
+++ b/Engine/Extras/VirtualProduction/LiveLinkVCAM/vcam/ja.lproj/Main.strings
@@ -5,10 +5,10 @@
 "76Q-l6-yNx.text" = "デモモード";
 
 /* Class = "UILabel"; text = "About"; ObjectID = "AF3-8v-IsQ"; */
-"AF3-8v-IsQ.text" = "About";
+"AF3-8v-IsQ.text" = "情報";
 
 /* Class = "UINavigationItem"; title = "About"; ObjectID = "AFZ-EN-fkq"; */
-"AFZ-EN-fkq.title" = "About";
+"AFZ-EN-fkq.title" = "情報";
 
 /* Class = "UILabel"; text = "PLACEHOLDER"; ObjectID = "aUB-pM-eQd"; */
 "aUB-pM-eQd.text" = "PLACEHOLDER";
