[EpicStageApp] Add editor config file for consistent formatting

diff --git a/Engine/Extras/VirtualProduction/EpicStageApp/.editorconfig b/Engine/Extras/VirtualProduction/EpicStageApp/.editorconfig
new file mode 100644
--- /dev/null
+++ b/Engine/Extras/VirtualProduction/EpicStageApp/.editorconfig
@@ -0,0 +1,11 @@
+# Copyright Epic Games, Inc. All Rights Reserved.
+
+root = true
+
+[*.dart]
+max_line_length = 120
+indent_style = space
+indent_size = 2
+tab_width = 2
+trim_trailing_whitespace = true
+insert_final_newline = true
