[EpicStageApp] Add press feedback and menu entry for actor focus action

diff --git a/Engine/Extras/VirtualProduction/EpicStageApp/android/app/src/debug/AndroidManifest.xml b/Engine/Extras/VirtualProduction/EpicStageApp/android/app/src/debug/AndroidManifest.xml
--- a/Engine/Extras/VirtualProduction/EpicStageApp/android/app/src/debug/AndroidManifest.xml
+++ b/Engine/Extras/VirtualProduction/EpicStageApp/android/app/src/debug/AndroidManifest.xml
@@ -4,4 +4,5 @@
          to allow setting breakpoints, to provide hot reload, etc.
     -->
     <uses-permission android:name="android.permission.INTERNET"/>
+    <uses-permission android:name="android.permission.VIBRATE" />
 </manifest>
diff --git a/Engine/Extras/VirtualProduction/EpicStageApp/android/app/src/main/AndroidManifest.xml b/Engine/Extras/VirtualProduction/EpicStageApp/android/app/src/main/AndroidManifest.xml
--- a/Engine/Extras/VirtualProduction/EpicStageApp/android/app/src/main/AndroidManifest.xml
+++ b/Engine/Extras/VirtualProduction/EpicStageApp/android/app/src/main/AndroidManifest.xml
@@ -37,4 +37,5 @@
     <uses-permission android:name="android.permission.INTERNET"/>
     <uses-permission android:name="android.permission.ACCESS_NETWORK_STATE" />
     <uses-permission android:name="android.permission.ACCESS_WIFI_STATE" />
+    <uses-permission android:name="android.permission.VIBRATE" />
 </manifest>
diff --git a/Engine/Extras/VirtualProduction/EpicStageApp/android/app/src/profile/AndroidManifest.xml b/Engine/Extras/VirtualProduction/EpicStageApp/android/app/src/profile/AndroidManifest.xml
--- a/Engine/Extras/VirtualProduction/EpicStageApp/android/app/src/profile/AndroidManifest.xml
+++ b/Engine/Extras/VirtualProduction/EpicStageApp/android/app/src/profile/AndroidManifest.xml
@@ -4,4 +4,5 @@
          to allow setting breakpoints, to provide hot reload, etc.
     -->
     <uses-permission android:name="android.permission.INTERNET"/>
+    <uses-permission android:name="android.permission.VIBRATE" />
 </manifest>
