[EpicStageApp] Localization support

diff --git a/Engine/Extras/VirtualProduction/EpicStageApp/l10n.yaml b/Engine/Extras/VirtualProduction/EpicStageApp/l10n.yaml
new file mode 100644
--- /dev/null
+++ b/Engine/Extras/VirtualProduction/EpicStageApp/l10n.yaml
@@ -0,0 +1,3 @@
+arb-dir: lib/l10n
+template-arb-file: app_en.arb
+output-localization-file: app_localizations.dart
\ No newline at end of file
diff --git a/Engine/Extras/VirtualProduction/EpicStageApp/pubspec.yaml b/Engine/Extras/VirtualProduction/EpicStageApp/pubspec.yaml
--- a/Engine/Extras/VirtualProduction/EpicStageApp/pubspec.yaml
+++ b/Engine/Extras/VirtualProduction/EpicStageApp/pubspec.yaml
@@ -43,10 +43,14 @@ dependencies:
   streaming_shared_preferences: ^2.0.0
   webview_flutter: ^3.0.4
   logging: ^1.1.0
-  intl: ^0.17.0
+  intl: any
   path_provider: ^2.0.11
   open_filex: ^4.3.1
   share_plus: ^6.3.0
+  flutter_localizations:
+    sdk: flutter
+  arb_generator: ^0.0.3
+  xml: ^6.2.2
 
 dev_dependencies:
   flutter_test:
@@ -66,6 +70,9 @@ dev_dependencies:
 
 # The following section is specific to Flutter.
 flutter:
+  
+  # Generate localization files
+  generate: true
 
   # The following line ensures that the Material Icons font is
   # included with your application, so that you can use the icons in
@@ -102,4 +109,13 @@ flutter_icons:
   android: true
   ios: true
   remove_alpha_ios: true
-  image_path: "assets/app_icon/EpicStage_1024x1024.png"
\ No newline at end of file
+  image_path: "assets/app_icon/EpicStage_1024x1024.png"
+  
+# Automatically generate arb localization files from a CSV
+arb_generator:
+  input_filepath: "lib/l10n/l10n.csv"
+  output_directory: "lib/l10n"
+  filename_prepend: "app_"
+  csv_settings:
+    description_index: 1
+    base_index: 2
\ No newline at end of file
