[UnrealStage] Add support for Remote Control passphrases

diff --git a/Engine/Extras/VirtualProduction/EpicStageApp/pubspec.yaml b/Engine/Extras/VirtualProduction/EpicStageApp/pubspec.yaml
--- a/Engine/Extras/VirtualProduction/EpicStageApp/pubspec.yaml
+++ b/Engine/Extras/VirtualProduction/EpicStageApp/pubspec.yaml
@@ -51,6 +51,7 @@ dependencies:
     sdk: flutter
   arb_generator: ^0.0.3
   xml: ^6.2.2
+  flutter_secure_storage: ^8.0.0
 
 dev_dependencies:
   flutter_test:
