[UnrealStage] Add documentation link

diff --git a/Engine/Extras/VirtualProduction/EpicStageApp/pubspec.yaml b/Engine/Extras/VirtualProduction/EpicStageApp/pubspec.yaml
--- a/Engine/Extras/VirtualProduction/EpicStageApp/pubspec.yaml
+++ b/Engine/Extras/VirtualProduction/EpicStageApp/pubspec.yaml
@@ -52,6 +52,7 @@ dependencies:
   arb_generator: ^0.0.3
   xml: ^6.2.2
   flutter_secure_storage: ^8.0.0
+  url_launcher: ^6.1.11
 
 dev_dependencies:
   flutter_test:
