Adding stats graphs to LiveLinkVCam app.
- Added a new stats subview
- Added a new line graph class
- Added a stats button that can be tapped to toggle stats visibility

diff --git a/Engine/Extras/VirtualProduction/LiveLinkVCAM/vcam/HeaderView.swift b/Engine/Extras/VirtualProduction/LiveLinkVCAM/vcam/HeaderView.swift
--- a/Engine/Extras/VirtualProduction/LiveLinkVCAM/vcam/HeaderView.swift
+++ b/Engine/Extras/VirtualProduction/LiveLinkVCAM/vcam/HeaderView.swift
@@ -12,6 +12,7 @@ import QuickLayout
 @objc protocol HeaderViewDelegate {
     func headerViewLogButtonTapped(_ headerView : HeaderView)
     func headerViewExitButtonTapped(_ headerView : HeaderView)
+    func headerViewStatsButtonTapped(_ headerView : HeaderView)
 }
 
 class HeaderView : UIView {
@@ -123,6 +124,7 @@ class HeaderView : UIView {
             btn.addTarget(self, action: #selector(exitTapped), for: .touchUpInside)
         }
 
+        // Show logs button
         if showButtons {
             
             let btn = UIButton(type: .custom)
@@ -140,6 +142,23 @@ class HeaderView : UIView {
             btn.addTarget(self, action: #selector(logTapped), for: .touchUpInside)
         }
         
+        // Show stats button
+        if showButtons {
+            let btn = UIButton(type: .custom)
+            btn.translatesAutoresizingMaskIntoConstraints = false
+            addSubview(btn)
+            
+            btn.imageView?.contentMode = .scaleAspectFit
+            btn.tintColor = UIColor.white
+            btn.setImage(UIImage(systemName: "chart.xyaxis.line"), for: .normal)
+
+            btn.set(.width, .height, of: 30)
+            btn.layoutToSuperview(.centerY)
+            btn.layoutToSuperview(.left, offset: 100)
+            
+            btn.addTarget(self, action: #selector(statsTapped), for: .touchUpInside)
+        }
+        
         if showButtons {
             
             let border = CALayer()
@@ -197,6 +216,12 @@ class HeaderView : UIView {
             d.headerViewLogButtonTapped(self)
         }
     }
+    
+    @objc func statsTapped(_ sender : Any) {
+        if let d = self.delegate {
+            d.headerViewStatsButtonTapped(self)
+        }
+    }
 
 
     @objc func displayLinkDidFire(_ displayLink: CADisplayLink) {
diff --git a/Engine/Extras/VirtualProduction/LiveLinkVCAM/vcam/LineGraphView.swift b/Engine/Extras/VirtualProduction/LiveLinkVCAM/vcam/LineGraphView.swift
new file mode 100644
--- /dev/null
+++ b/Engine/Extras/VirtualProduction/LiveLinkVCAM/vcam/LineGraphView.swift
@@ -0,0 +1,195 @@
+//
+//  LineGraphView.swift
+//  vcam
+//
+//  Created by TensorWorks on 17/7/2023.
+//  Copyright Epic Games, Inc. All Rights Reserved.
+//
+
+import Foundation
+import UIKit
+
+class LineGraphView: UIView {
+    
+    // An array of data points for the graph
+    private var dataPoints: [CGFloat] = []
+    
+    // Limit to the number of data points stored in the array
+    private var sizeLimit: UInt16 = 120;
+    
+    // Font to use in our graph title
+    private var titleFont: UIFont = UIFont.systemFont(ofSize: 12);
+    
+    // Font to use in our graph tick labels
+    private var tickLabelFont: UIFont = UIFont.systemFont(ofSize: 8);
+    
+    // Color of our font
+    private var fontColor: UIColor = UIColor.white;
+    
+    // Paragraph style for our text
+    private var paragraphStyle: NSMutableParagraphStyle = NSMutableParagraphStyle();
+    
+    // Title text attributes
+    private var titleTextAttributes: [NSAttributedString.Key: Any] = [:];
+    
+    // Tick label text attributes
+    private var tickTextAttributes: [NSAttributedString.Key: Any] = [:];
+    
+    // Maximum y value we have seen
+    public var maxY: CGFloat = 1;
+    
+    // Title of graph
+    public var title: String = "Title" {
+        didSet {
+            setNeedsDisplay() // Trigger redraw when the text is set
+        }
+    }
+    
+    // Function for formatting the numberic text labels next to the y-Axis ticks
+    public var yTickLabelFormatter : (Double) -> String = {
+        return String(format: "%.0f", round($0))
+    }
+    
+    override init(frame: CGRect) {
+        super.init(frame: frame)
+        
+        // set transparent background color
+        self.backgroundColor = UIColor(red: 0.0, green: 0.0, blue: 0.0, alpha: 0.5)
+        
+        // fill data points with 1's
+        for _ in 1...self.sizeLimit {
+            addDataPoint(value: CGFloat(0));
+        }
+        
+        // Set alignment of text in our graphs to be center
+        self.paragraphStyle.alignment = .left;
+        
+        // Set the attributes for the text
+        self.titleTextAttributes = [
+            .font: self.titleFont,
+            .foregroundColor: self.fontColor,
+            .paragraphStyle: self.paragraphStyle
+        ]
+        
+        self.tickTextAttributes = [
+            .font: self.tickLabelFont,
+            .foregroundColor: self.fontColor,
+            .paragraphStyle: self.paragraphStyle
+        ]
+    }
+    
+    required init?(coder: NSCoder) {
+        super.init(coder: coder)
+        self.backgroundColor = UIColor.clear
+    }
+    
+    func addDataPoint(value: CGFloat) {
+        
+        if value > self.maxY || value == self.maxY {
+            // adds some headroom the the graph
+            self.maxY = value + (self.maxY * 0.1)
+        }
+                
+        self.dataPoints.append(value)
+        
+        if self.dataPoints.count > self.sizeLimit {
+            self.dataPoints.removeFirst();
+        }
+        
+        // Triggers a redraw (need to do this from main thread)
+        DispatchQueue.main.async {
+            self.setNeedsDisplay();
+        }
+    }
+    
+    override func draw(_ rect: CGRect) {
+        super.draw(rect)
+        
+        guard let context = UIGraphicsGetCurrentContext() else { return }
+        
+        // Draw the text in the bounding rectangle
+        let titleSize = self.title.size(withAttributes: self.titleTextAttributes);
+        let titleRect = CGRect(
+            x: (rect.width - titleSize.width)  * 0.5,
+            y: 0,
+            width: rect.width,
+            height: rect.height
+        )
+        self.title.draw(in: titleRect, withAttributes: self.titleTextAttributes)
+        
+
+        // Draw the y-axis
+        let axisPadding = rect.width * 0.125;
+        UIColor.white.setStroke()
+        let yAxis = UIBezierPath()
+        yAxis.lineWidth = 1.0
+        yAxis.move(to: CGPoint(x: axisPadding, y: titleSize.height))
+        yAxis.addLine(to: CGPoint(x: axisPadding, y: rect.height))
+        yAxis.stroke()
+        
+        // Draw y-axis ticks
+        let tickLength = 3.0
+        let yAxisHeight = rect.height - titleSize.height
+        let nTicks = 5
+        
+        for tickYIdx in 0...nTicks {
+            let tickPath = UIBezierPath();
+            tickPath.lineWidth = 1.0
+            let scalar = CGFloat(tickYIdx) / CGFloat(nTicks)
+            let tickY = titleSize.height + (scalar * yAxisHeight)
+            tickPath.move(to: CGPoint(x: axisPadding, y: tickY))
+            tickPath.addLine(to: CGPoint(x: axisPadding - tickLength, y: tickY))
+            tickPath.stroke()
+            // Draw tick label
+            let tickValue : CGFloat = self.maxY - ( self.maxY * scalar )
+            let tickLabel = self.yTickLabelFormatter(tickValue)
+            
+            let tickLabelSize = tickLabel.size(withAttributes: self.tickTextAttributes);
+            let tickLabelX : CGFloat = axisPadding - tickLength - tickLabelSize.width
+            let tickLabelYOffset = tickYIdx == nTicks ? tickLabelSize.height : tickLabelSize.height * 0.5
+            let tickLabelY: CGFloat = tickY - tickLabelYOffset
+            let tickRect = CGRect(x: tickLabelX, y: tickLabelY, width: rect.width, height: rect.height)
+            tickLabel.draw(in: tickRect, withAttributes: self.tickTextAttributes)
+        }
+        
+        // Draw the line graph
+        let graphYSize = rect.height - titleSize.height
+        let graphXSize = rect.width - axisPadding
+        
+        // Set the line color
+        UIColor.red.setStroke()
+        
+        // Set up the coordinate system
+        let coordinateSystem = CGRect(x: rect.origin.x, y: rect.origin.y, width: graphXSize, height: graphYSize)
+        context.saveGState()
+        context.translateBy(x: axisPadding, y: coordinateSystem.height + titleSize.height)
+        context.scaleBy(x: 1, y: -1)
+        
+        // Calculate the scale factors
+        let maxX = CGFloat(dataPoints.count)
+        let scaleX = coordinateSystem.width / maxX
+        let scaleY = coordinateSystem.height / self.maxY
+        
+        // Create the line path
+        let path = UIBezierPath()
+        path.lineWidth = 1.0
+        
+        for (index, dataPoint) in dataPoints.enumerated() {
+            let x = CGFloat(index) * scaleX
+            let y = dataPoint * scaleY
+            
+            if index == 0 {
+                path.move(to: CGPoint(x: x, y: y))
+            } else {
+                path.addLine(to: CGPoint(x: x, y: y))
+            }
+        }
+        
+        // Draw the line
+        path.stroke()
+        
+        context.restoreGState()
+    }
+    
+}
+
diff --git a/Engine/Extras/VirtualProduction/LiveLinkVCAM/vcam/StreamingConnection.swift b/Engine/Extras/VirtualProduction/LiveLinkVCAM/vcam/StreamingConnection.swift
--- a/Engine/Extras/VirtualProduction/LiveLinkVCAM/vcam/StreamingConnection.swift
+++ b/Engine/Extras/VirtualProduction/LiveLinkVCAM/vcam/StreamingConnection.swift
@@ -119,6 +119,10 @@ class StreamingConnection : NSObject {
         self.subjectName = subjectName
     }
     
+    func showStats(_ shouldShow : Bool) {
+        // Do nothing, derived classes will implement functionality
+    }
+    
     func shutdown() {
         assertionFailure("not implemented")
     }
diff --git a/Engine/Extras/VirtualProduction/LiveLinkVCAM/vcam/VideoViewController.swift b/Engine/Extras/VirtualProduction/LiveLinkVCAM/vcam/VideoViewController.swift
--- a/Engine/Extras/VirtualProduction/LiveLinkVCAM/vcam/VideoViewController.swift
+++ b/Engine/Extras/VirtualProduction/LiveLinkVCAM/vcam/VideoViewController.swift
@@ -41,6 +41,7 @@ class VideoViewController : BaseViewController {
     var forceDisconnectTimer : Timer?
     
     var statsTimer : Timer?
+    var showStats : Bool = false
     
     var gameControllerSnapshot : GCController?
     weak var gameController : GCController? {
@@ -255,6 +256,7 @@ class VideoViewController : BaseViewController {
              
                 str += "\(fps) fps"
             }
+            
         }
 
         self.headerView.stats = str
@@ -444,6 +446,11 @@ class VideoViewController : BaseViewController {
 }
 
 extension VideoViewController : HeaderViewDelegate {
+
+    func headerViewStatsButtonTapped(_ headerView: HeaderView) {
+        self.showStats = !self.showStats
+        streamingConnection?.showStats(self.showStats)
+    }
     
     func headerViewExitButtonTapped(_ headerView : HeaderView) {
      
@@ -459,6 +466,7 @@ extension VideoViewController : HeaderViewDelegate {
      
         performSegue(withIdentifier: "showLog", sender: headerView)
     }
+    
 }
 
 
diff --git a/Engine/Extras/VirtualProduction/LiveLinkVCAM/vcam/WebRTCStats.swift b/Engine/Extras/VirtualProduction/LiveLinkVCAM/vcam/WebRTCStats.swift
new file mode 100644
--- /dev/null
+++ b/Engine/Extras/VirtualProduction/LiveLinkVCAM/vcam/WebRTCStats.swift
@@ -0,0 +1,220 @@
+//
+//  WebRTCStats.swift
+//  vcam
+//
+//  Created by TensorWorks on 24/7/2023.
+//  Copyright Epic Games, Inc. All Rights Reserved.
+//
+
+import Foundation
+import WebRTC
+
+
+class WebRTCStats
+{
+    public var lastBitrate : Int64?
+    public var lastFPS : Int64?
+    private var lastFramesDecoded : Int64?
+    private var lastBytesReceived : Int64?
+    private var lastBytesReceivedTimestamp : CFTimeInterval?
+    private var lastInterframeDelay : Double?
+    private var lastProcessingDelay : Double?
+    
+    private var statsView : WebRTCStatsView
+    private var fpsGraph : LineGraphView?
+    private var framesReceivedGraph : LineGraphView?
+    private var bitrateGraph : LineGraphView?
+    private var widthGraph : LineGraphView?
+    private var heightGraph : LineGraphView?
+    private var jitterGraph : LineGraphView?
+    private var jitterBufferDelayGraph : LineGraphView?
+    private var multipacketFramesGraph : LineGraphView?
+    private var freezeCountGraph : LineGraphView?
+    private var nackCountGraph : LineGraphView?
+    private var pliCountGraph : LineGraphView?
+    private var packetLossGraph : LineGraphView?
+    private var keyframesGraph : LineGraphView?
+    private var framesDecodedGraph : LineGraphView?
+    private var framesDroppedGraph : LineGraphView?
+    private var interframeDelayGraph : LineGraphView?
+    private var processingTimeGraph : LineGraphView?
+    
+    init(statsView: WebRTCStatsView) {
+        
+        self.statsView = statsView
+        
+        self.bitrateGraph = self.statsView.addGraph(graphName: "Bitrate (megabit/s)")
+        self.bitrateGraph?.maxY = 20
+        
+        self.fpsGraph = self.statsView.addGraph(graphName: "FPS")
+        
+        self.framesReceivedGraph = self.statsView.addGraph(graphName: "Frames received")
+        self.framesReceivedGraph?.maxY = 100
+        
+        self.multipacketFramesGraph = self.statsView.addGraph(graphName: "Multi-packet Frames")
+        self.multipacketFramesGraph?.maxY = 10
+        		
+        self.jitterGraph = self.statsView.addGraph(graphName: "Jitter")
+        self.jitterGraph?.maxY = 0.05
+        self.jitterGraph?.yTickLabelFormatter = { return String(format: "%.2f", $0) }
+        
+        self.jitterBufferDelayGraph = self.statsView.addGraph(graphName: "Jitter Buffer Delay(ms)")
+
+        self.packetLossGraph = self.statsView.addGraph(graphName: "Packets lost")
+        self.packetLossGraph?.maxY = 10
+        
+        self.nackCountGraph = self.statsView.addGraph(graphName: "Nack Count")
+        self.nackCountGraph?.maxY = 10
+
+        self.freezeCountGraph = self.statsView.addGraph(graphName: "Freeze Count")
+        self.freezeCountGraph?.maxY = 10
+        
+        self.pliCountGraph = self.statsView.addGraph(graphName: "Picture Loss Count")
+        self.pliCountGraph?.maxY = 10
+        
+        self.keyframesGraph = self.statsView.addGraph(graphName: "Keyframes Decoded")
+        self.keyframesGraph?.maxY = 10
+        
+        self.framesDecodedGraph = self.statsView.addGraph(graphName: "Decode FPS")
+        self.framesDecodedGraph?.maxY = 10
+        
+        self.framesDroppedGraph = self.statsView.addGraph(graphName: "Frames Dropped")
+        self.framesDroppedGraph?.maxY = 10
+        
+        self.interframeDelayGraph = self.statsView.addGraph(graphName: "Interframe delay (ms)")
+        self.interframeDelayGraph?.maxY = 10
+
+        self.processingTimeGraph = self.statsView.addGraph(graphName: "Processing delay (ms)")
+        self.processingTimeGraph?.maxY = 10
+        
+    }
+    
+    func processStatsReport(report: RTCStatisticsReport){
+        
+        for (_,value) in report.statistics {
+            if let kind = value.values["kind"] as? String {
+                if kind == "video" {
+                    
+                    //Log.info("\(value.values)")
+                    
+                    // Jitter buffer delay
+                    if let jitterBufferDelay = value.values["jitterBufferDelay"] as? Double, let jitterBufferEmittedCount = value.values["jitterBufferEmittedCount"] as? UInt64 {
+                        let jbDelayMs : Double = jitterBufferDelay / Double(jitterBufferEmittedCount) * 1000.0
+                        self.jitterBufferDelayGraph?.addDataPoint(value: jbDelayMs)
+                    }
+                    
+                    // Jitter
+                    if let jitter = value.values["jitter"] as? Double {
+                        self.jitterGraph?.addDataPoint(value: jitter)
+                    }
+
+                    // FPS
+                    if let fps = value.values["framesPerSecond"] as? Double {
+                        self.lastFPS = Int64(fps)
+                        self.fpsGraph?.addDataPoint(value: fps)
+                    }
+                    
+                    // Frames received
+                    if let framesReceived = value.values["framesReceived"] as? Int32 {
+                        self.framesReceivedGraph?.addDataPoint(value: Double(framesReceived))
+                    }
+                    
+                    // framesAssembledFromMultiplePackets
+                    if let multipacketFrames = value.values["framesAssembledFromMultiplePackets"] as? Int64 {
+                        self.multipacketFramesGraph?.addDataPoint(value: Double(multipacketFrames))
+                    }
+                    
+                    // Freeze Count
+                    if let freezeCount = value.values["freezeCount"] as? Int64 {
+                        self.freezeCountGraph?.addDataPoint(value: Double(freezeCount))
+                    }
+                    
+                    // Nack Count
+                    if let nackCount = value.values["nackCount"] as? Int64 {
+                        self.nackCountGraph?.addDataPoint(value: Double(nackCount))
+                    }
+                    
+                    // PLI count
+                    if let pliCount = value.values["pliCount"] as? Int64 {
+                        self.pliCountGraph?.addDataPoint(value: Double(pliCount))
+                    }
+                    
+                    // Bitrate
+                    if let bytesReceived = value.values["bytesReceived"] as? Int64 {
+                        
+                        if let lastBytes = self.lastBytesReceived {
+                            self.lastBitrate = bytesReceived - lastBytes
+                            let megabitsDelta : Double = Double(self.lastBitrate!) / 125000.0
+                            self.bitrateGraph?.addDataPoint(value: megabitsDelta)
+                        }
+                        
+                        self.lastBytesReceived = bytesReceived
+                    }
+                    
+                    // Packets lost
+                    if let packetsLost = value.values["packetsLost"] as? Int32 {
+                        self.packetLossGraph?.addDataPoint(value: Double(packetsLost))
+                    }
+                    
+                    // Keyframes decoded
+                    if let keyframesDecoded = value.values["keyFramesDecoded"] as? Int64 {
+                        self.keyframesGraph?.addDataPoint(value: Double(keyframesDecoded))
+                    }
+                    
+                    // Frames decoded
+                    if let framesDecoded = value.values["framesDecoded"] as? Int64 {
+                        if let lastFrames = self.lastFramesDecoded {
+                            let framesDelta : Int64 = framesDecoded - lastFrames;
+                            if framesDelta > 0 {
+                                self.framesDecodedGraph?.addDataPoint(value: Double(framesDelta))
+                            }
+                        }
+                        self.lastFramesDecoded = framesDecoded
+                    }
+                    
+                    // Frames dropped
+                    if let framesDropped = value.values["framesDropped"] as? Int64 {
+                        self.framesDroppedGraph?.addDataPoint(value: Double(framesDropped))
+                    }
+                    
+                    // totalInterFrameDelay
+                    if let totalInterFrameDelay = value.values["totalInterFrameDelay"] as? Double {
+                        if let lastDelay = self.lastInterframeDelay {
+                            let delayDelta : Double = totalInterFrameDelay - lastDelay
+                            if delayDelta > 0 {
+                                self.interframeDelayGraph?.addDataPoint(value: Double(delayDelta))
+                            }
+                        }
+                        self.lastInterframeDelay = totalInterFrameDelay
+                    }
+                    
+                    // totalProcessingDelay
+                    if let totalProcessingDelay = value.values["totalProcessingDelay"] as? Double {
+                        if let lastProcessing = self.lastProcessingDelay {
+                            let processingDelta : Double = totalProcessingDelay - lastProcessing
+                            if processingDelta > 0 {
+                                self.processingTimeGraph?.addDataPoint(value: Double(processingDelta))
+                            }
+                        }
+                        self.lastProcessingDelay = totalProcessingDelay
+                    }
+
+                    // There is also:
+                    // firCount
+                    // pauseCount
+                    // totalSquaredInterFrameDelay : Double
+                    // totalDecodeTime : Double
+                    // totalPausesDuration : Double
+                    // packetsReceived : Int64
+                    // totalFreezesDuration : Int64
+                    // These were excluded because we can only fit so many graphs on the screen at once
+                    // and these ones were either inferred from other graphs or less likely to be problematic (we hope)
+                    
+                }
+            }
+        }
+        
+        
+    }
+    
+}
diff --git a/Engine/Extras/VirtualProduction/LiveLinkVCAM/vcam/WebRTCStatsView.swift b/Engine/Extras/VirtualProduction/LiveLinkVCAM/vcam/WebRTCStatsView.swift
new file mode 100644
--- /dev/null
+++ b/Engine/Extras/VirtualProduction/LiveLinkVCAM/vcam/WebRTCStatsView.swift
@@ -0,0 +1,70 @@
+//
+//  WebRTCStatsView.swift
+//  vcam
+//
+//  Created by TensorWorks on 17/7/2023.
+//  Copyright Epic Games, Inc. All Rights Reserved.
+//
+
+import UIKit
+
+// Draws WebRTC stats as line graphs
+class WebRTCStatsView: UIView {
+    
+    private var statsGraphs : [String: LineGraphView] = [:]
+    private var graphWidth : Int = 175
+    private var graphHeight : Int = 175
+    private var graphPadding : Int = 10
+    private var graphHeaderYPadding : Int = 30
+    private var graphLeftMargin : Int = 100;
+    private var graphRightMargin : Int = 100;
+
+    // Graphs will be displayed in a grid layout, rowsxcols, where nRows and nCols are determined by frameWidth/graphWidth
+    // When graphs exceed nCols on a single row then a new row will begin
+    // There is no maximum number of rows, but some graphs may be drawn offscreen if there is too many rows
+    private var nCols : Int = 5
+    
+    override init(frame: CGRect) {
+        super.init(frame: frame)
+        
+        // Disable user interaction so stats do not block underlying video
+        self.isUserInteractionEnabled = false
+        
+        // set transparent background color
+        self.backgroundColor = UIColor.clear
+        
+        // Calculate nRows and nCols
+        let frameWidth : CGFloat = frame.width - CGFloat(self.graphLeftMargin) - CGFloat(self.graphRightMargin)
+        self.nCols = Int( floor( frameWidth / CGFloat(self.graphWidth) ))
+
+    }
+    
+    required init?(coder: NSCoder) {
+        super.init(coder: coder)
+        
+    }
+    
+    func addGraph(graphName: String) -> LineGraphView {
+        
+        // Calculate the rect in which the graph is drawn
+        let nGraphs : Int = self.statsGraphs.count
+        let colIdx : Int = nGraphs % self.nCols
+        let rowIdx : Int = abs(nGraphs /  self.nCols)
+        let graphX : Int = colIdx * (self.graphWidth + self.graphPadding) + self.graphLeftMargin
+        let graphY : Int = rowIdx * (self.graphHeight + self.graphPadding) + self.graphHeaderYPadding
+        let graphFrame : CGRect = CGRect(x: graphX, y: graphY, width: self.graphWidth, height: self.graphHeight)
+        let graph : LineGraphView = LineGraphView(frame: graphFrame)
+
+        // Add the graph the be drawn under this view
+        self.addSubview(graph)
+        
+        // Given it the passed in key as its title
+        graph.title = graphName
+        
+        // Add the graph to our map of graphs
+        self.statsGraphs[graphName] = graph
+        
+        return graph
+    }
+    
+}
diff --git a/Engine/Extras/VirtualProduction/LiveLinkVCAM/vcam/WebRTCStreamingConnection.swift b/Engine/Extras/VirtualProduction/LiveLinkVCAM/vcam/WebRTCStreamingConnection.swift
--- a/Engine/Extras/VirtualProduction/LiveLinkVCAM/vcam/WebRTCStreamingConnection.swift
+++ b/Engine/Extras/VirtualProduction/LiveLinkVCAM/vcam/WebRTCStreamingConnection.swift
@@ -18,6 +18,7 @@ class WebRTCStreamingConnection : StreamingConnection {
     private var touchControls: TouchControls?
     private var keyboardControls: KeyboardControls?
     private var webRTCView : WebRTCView?
+    private var webRTCStatsView : WebRTCStatsView?
     private var rtcVideoTrack : RTCVideoTrack?
     
     private var signalingConnected = false
@@ -25,9 +26,9 @@ class WebRTCStreamingConnection : StreamingConnection {
     private var hasLocalSdp = false
     private var remoteCandidateCount = 0
     private var localCandidateCount = 0
+    
+    private var webRTCStats : WebRTCStats?
     private var _statsTimer : Timer?
-    private var _lastBytesReceived : Int?
-    private var _lastBytesReceivedTimestamp : CFTimeInterval?
     
     private var reconnectAttempt : Int = 0
     private var maxReconnectAttempts : Int = 3
@@ -67,8 +68,22 @@ class WebRTCStreamingConnection : StreamingConnection {
                 rtcView.delegate = self
                 rtcView.layoutToSuperview(.top, .bottom, .left, .right)
                 self.webRTCView = rtcView
-                
                 self.attachVideoTrack()
+                
+                // Attach stats view here
+                let rtcStatsView = WebRTCStatsView(frame: CGRect(x: 0, y: 0, width: rv.frame.size.width, height: rv.frame.size.height));
+                rv.addSubview(rtcStatsView)
+                NSLayoutConstraint.activate([
+                    rtcStatsView.topAnchor.constraint(equalTo: rv.topAnchor),
+                    rtcStatsView.leadingAnchor.constraint(equalTo: rv.leadingAnchor),
+                    rtcStatsView.trailingAnchor.constraint(equalTo: rv.trailingAnchor),
+                    rtcStatsView.bottomAnchor.constraint(equalTo: rv.bottomAnchor)
+                ]);
+                // start with stats hidden
+                rtcStatsView.isHidden = true
+                self.webRTCStatsView = rtcStatsView;
+                self.webRTCStats = WebRTCStats(statsView: rtcStatsView)
+                
             }
         }
     }
@@ -79,6 +94,8 @@ class WebRTCStreamingConnection : StreamingConnection {
         self.webRTCClient = WebRTCClient()
         self.webRTCClient?.speakerOff()
         self.webRTCClient?.delegate = self
+        
+        self.stats = StreamingConnectionStats()
 
         _statsTimer = Timer.scheduledTimer(withTimeInterval: 1, repeats: true, block: { timer in
             
@@ -86,38 +103,11 @@ class WebRTCStreamingConnection : StreamingConnection {
                 
                 webRTC.stats({ report in
                     
-                    for (key,value) in report.statistics {
-                        if key.starts(with: "RTCIceCandidatePair_") {
-                            //Log.info("\(value.values)")
-                            let timestamp = value.timestamp_us
-                            if let bytesReceived = value.values["bytesReceived"] as? Int {
-                                
-                                if bytesReceived > 0 {
-
-                                    var bytesPerSecond : Int?
-                                    if let previousTimestamp = self._lastBytesReceivedTimestamp,
-                                       let previousBytesReceived = self._lastBytesReceived {
-
-                                        bytesPerSecond = Int(Double(bytesReceived - previousBytesReceived) / ((timestamp - previousTimestamp) / 1000000.0))
-                                    }
-
-                                    self._lastBytesReceivedTimestamp = timestamp
-                                    self._lastBytesReceived = bytesReceived
-                                    
-                                    if let bps = bytesPerSecond {
-                                        self.stats = StreamingConnectionStats()
-                                        self.stats?.bytesPerSecond = bps
-                                    }
-                                }
-                            }
-                            
-                            break
-                        }
-                    }
+                    self.webRTCStats?.processStatsReport(report: report)
+                    self.stats?.framesPerSecond = Float(self.webRTCStats?.lastFPS ?? 0)
+                    self.stats?.bytesPerSecond = Int(self.webRTCStats?.lastBitrate ?? 0)
                 })
             } else {
-                self._lastBytesReceived = nil
-                self._lastBytesReceivedTimestamp = nil
                 self.stats = nil
             }
         })
@@ -129,6 +119,10 @@ class WebRTCStreamingConnection : StreamingConnection {
 
         webRTCClient = nil
     }
+    
+    override func showStats(_ shouldShow : Bool) {
+        self.webRTCStatsView?.isHidden = shouldShow == false
+    }
 
     override func shutdown() {
         _statsTimer?.invalidate()
@@ -445,11 +439,7 @@ extension WebRTCStreamingConnection: SignalClientDelegate {
             } else {
                 // We've exhausted our reconnect attempts, return to main menu
                 self.reconnectAttempt = 0
-                self.delegate?.streamingConnection(self, exitWithError:
-                                                    NSError(domain: "",
-                                                            code: 1,
-                                                            userInfo: [NSLocalizedDescriptionKey:
-                                                                        String(format: NSLocalizedString("error-unablereconnect", value:"Unable to reconnect to %s after %d attempts.", comment: "Error message : %s is replace with a string and %d with a number"), subscribedStreamer, maxReconnectAttempts)]))
+                self.delegate?.streamingConnection(self, exitWithError: NSError(domain: "", code: 1, userInfo: [NSLocalizedDescriptionKey: "Unable to reconnect to \(subscribedStreamer) after \(maxReconnectAttempts) attempts"]))
             }
             
         } else {
