Fix: Stop frame stuttering at large resolutions by increase packet buffer size

diff --git a/Engine/Extras/VirtualProduction/LiveLinkVCAM/vcam/WebRTC/Services/WebRTCClient.swift b/Engine/Extras/VirtualProduction/LiveLinkVCAM/vcam/WebRTC/Services/WebRTCClient.swift
--- a/Engine/Extras/VirtualProduction/LiveLinkVCAM/vcam/WebRTC/Services/WebRTCClient.swift
+++ b/Engine/Extras/VirtualProduction/LiveLinkVCAM/vcam/WebRTC/Services/WebRTCClient.swift
@@ -25,6 +25,11 @@ final class WebRTCClient: NSObject {
         RTCInitializeSSL()
         let videoEncoderFactory = RTCDefaultVideoEncoderFactory()
         let videoDecoderFactory = RTCDefaultVideoDecoderFactory()
+
+        let fieldTrials = ["WebRTC-MaxPacketBufferSize", "8192"]
+        RTCInitFieldTrialDictionary(fieldTrials)
+
+
         return RTCPeerConnectionFactory(encoderFactory: videoEncoderFactory, decoderFactory: videoDecoderFactory)
     }()
     
