Live Link VCAM iOS (v1.3.2 build 1)

diff --git a/Engine/Extras/VirtualProduction/LiveLinkVCAM/vcam/WebRTC/Services/WebRTCClient.swift b/Engine/Extras/VirtualProduction/LiveLinkVCAM/vcam/WebRTC/Services/WebRTCClient.swift
--- a/Engine/Extras/VirtualProduction/LiveLinkVCAM/vcam/WebRTC/Services/WebRTCClient.swift
+++ b/Engine/Extras/VirtualProduction/LiveLinkVCAM/vcam/WebRTC/Services/WebRTCClient.swift
@@ -26,7 +26,7 @@ final class WebRTCClient: NSObject {
         let videoEncoderFactory = RTCDefaultVideoEncoderFactory()
         let videoDecoderFactory = RTCDefaultVideoDecoderFactory()
 
-        let fieldTrials = ["WebRTC-MaxPacketBufferSize", "8192"]
+        let fieldTrials = ["WebRTC-MaxPacketBufferSize": "8192"]
         RTCInitFieldTrialDictionary(fieldTrials)
 
 
