Branch snapshot for CL 27138787

[CL 27138787 in ue5-main branch]

diff --git a/.gitignore b/.gitignore
--- a/.gitignore
+++ b/.gitignore
@@ -162,6 +162,9 @@
 /Engine/Extras/VirtualProduction/LiveLinkVCAM/UE4LiveLinkProvider.framework/UE4LiveLinkProvider
 /Engine/Extras/VirtualProduction/LiveLinkVCAM/UE5LiveLinkProvider.framework/UE5LiveLinkProvider
 
+# Exception for VP EpicStageApp dart files
+!/Engine/Extras/VirtualProduction/EpicStageApp/**/*.dart
+
 # Ignore the automatically created mono symlink
 /Engine/Binaries/ThirdParty/Mono/Mac/bin/mono
 
diff --git a/Engine/Extras/VirtualProduction/EpicStageApp/lib/gestures/rotate_and_scale.dart b/Engine/Extras/VirtualProduction/EpicStageApp/lib/gestures/rotate_and_scale.dart
new file mode 100644
--- /dev/null
+++ b/Engine/Extras/VirtualProduction/EpicStageApp/lib/gestures/rotate_and_scale.dart
@@ -0,0 +1,411 @@
+// Copyright Epic Games, Inc. All Rights Reserved.
+
+// Contents of this file largely based on flutter/lib/src/gestures/scale.dart
+// This variation of the scale gesture has been adapted to immediately accept the gesture as soon as 2+ pointers are
+// detected.
+
+import 'dart:math' as math;
+
+import 'package:flutter/gestures.dart';
+
+/// The possible states of a [RotateAndScaleGestureRecognizer].
+enum _ScaleAndRotateState {
+  /// The recognizer is ready to start recognizing a gesture.
+  ready,
+
+  /// The sequence of pointer events seen thus far is consistent with a scale
+  /// gesture but the gesture has not been accepted definitively.
+  possible,
+
+  /// The sequence of pointer events seen thus far has been accepted
+  /// definitively as a scale gesture.
+  accepted,
+
+  /// The sequence of pointer events seen thus far has been accepted
+  /// definitively as a scale gesture and the pointers established a focal point
+  /// and initial scale.
+  started,
+}
+
+bool _isFlingGesture(Velocity velocity) {
+  final double speedSquared = velocity.pixelsPerSecond.distanceSquared;
+  return speedSquared > kMinFlingVelocity * kMinFlingVelocity;
+}
+
+/// Defines a line between two pointers on screen.
+///
+/// [_LineBetweenPointers] is an abstraction of a line between two pointers in
+/// contact with the screen. Used to track the rotation of a scale gesture.
+class _LineBetweenPointers {
+  /// Creates a [_LineBetweenPointers]. None of the [pointerStartLocation], [pointerStartId]
+  /// [pointerEndLocation] and [pointerEndId] must be null. [pointerStartId] and [pointerEndId]
+  /// should be different.
+  _LineBetweenPointers({
+    this.pointerStartLocation = Offset.zero,
+    this.pointerStartId = 0,
+    this.pointerEndLocation = Offset.zero,
+    this.pointerEndId = 1,
+  }) : assert(pointerStartId != pointerEndId);
+
+  // The location and the id of the pointer that marks the start of the line.
+  final Offset pointerStartLocation;
+  final int pointerStartId;
+
+  // The location and the id of the pointer that marks the end of the line.
+  final Offset pointerEndLocation;
+  final int pointerEndId;
+}
+
+/// Recognizes a scale and/or rotate gesture.
+/// Copied from Flutter's built-in ScaleAndRotateGestureRecognizer, but modified to accept the gesture based on angle of
+/// rotation and a lower scale slop value.
+///
+/// [RotateAndScaleGestureRecognizer] tracks the pointers in contact with the screen and
+/// calculates their focal point, indicated scale, and rotation. When a focal
+/// pointer is established, the recognizer calls [onStart]. As the focal point,
+/// scale, rotation change, the recognizer calls [onUpdate]. When the pointers
+/// are no longer in contact with the screen, the recognizer calls [onEnd].
+class RotateAndScaleGestureRecognizer extends OneSequenceGestureRecognizer {
+  /// Create a gesture recognizer for interactions intended for scaling content.
+  ///
+  /// {@macro flutter.gestures.GestureRecognizer.supportedDevices}
+  RotateAndScaleGestureRecognizer({
+    Object? debugOwner,
+    Set<PointerDeviceKind>? supportedDevices,
+    this.dragStartBehavior = DragStartBehavior.down,
+  }) : super(debugOwner: debugOwner, supportedDevices: supportedDevices);
+
+  /// Determines what point is used as the starting point in all calculations
+  /// involving this gesture.
+  ///
+  /// When set to [DragStartBehavior.down], the scale is calculated starting
+  /// from the position where the pointer first contacted the screen.
+  ///
+  /// When set to [DragStartBehavior.start], the scale is calculated starting
+  /// from the position where the scale gesture began. The scale gesture may
+  /// begin after the time that the pointer first contacted the screen if there
+  /// are multiple listeners competing for the gesture. In that case, the
+  /// gesture arena waits to determine whether or not the gesture is a scale
+  /// gesture before giving the gesture to this GestureRecognizer. This happens
+  /// in the case of nested GestureDetectors, for example.
+  ///
+  /// Defaults to [DragStartBehavior.down].
+  ///
+  /// See also:
+  ///
+  /// * https://flutter.dev/docs/development/ui/advanced/gestures#gesture-disambiguation,
+  ///   which provides more information about the gesture arena.
+  DragStartBehavior dragStartBehavior;
+
+  /// The pointers in contact with the screen have established a focal point and
+  /// initial scale of 1.0.
+  ///
+  /// This won't be called until the gesture arena has determined that this
+  /// GestureRecognizer has won the gesture.
+  ///
+  /// See also:
+  ///
+  /// * https://flutter.dev/docs/development/ui/advanced/gestures#gesture-disambiguation,
+  ///   which provides more information about the gesture arena.
+  GestureScaleStartCallback? onStart;
+
+  /// The pointers in contact with the screen have indicated a new focal point
+  /// and/or scale.
+  GestureScaleUpdateCallback? onUpdate;
+
+  /// The pointers are no longer in contact with the screen.
+  GestureScaleEndCallback? onEnd;
+
+  _ScaleAndRotateState _state = _ScaleAndRotateState.ready;
+
+  Matrix4? _lastTransform;
+
+  Offset? _currentFocalPoint;
+  late double _initialSpan;
+  late double _currentSpan;
+  late double _initialHorizontalSpan;
+  late double _currentHorizontalSpan;
+  late double _initialVerticalSpan;
+  late double _currentVerticalSpan;
+  late Offset _localFocalPoint;
+  _LineBetweenPointers? _initialLine;
+  _LineBetweenPointers? _currentLine;
+  late Map<int, Offset> _pointerLocations;
+  late List<int> _pointerQueue; // A queue to sort pointers in order of entrance
+  final Map<int, VelocityTracker> _velocityTrackers = <int, VelocityTracker>{};
+  late Offset _delta;
+
+  double get _scaleFactor => _initialSpan > 0.0 ? _currentSpan / _initialSpan : 1.0;
+
+  double get _horizontalScaleFactor =>
+      _initialHorizontalSpan > 0.0 ? _currentHorizontalSpan / _initialHorizontalSpan : 1.0;
+
+  double get _verticalScaleFactor => _initialVerticalSpan > 0.0 ? _currentVerticalSpan / _initialVerticalSpan : 1.0;
+
+  double _computeRotationFactor() {
+    if (_initialLine == null || _currentLine == null) {
+      return 0.0;
+    }
+    final double fx = _initialLine!.pointerStartLocation.dx;
+    final double fy = _initialLine!.pointerStartLocation.dy;
+    final double sx = _initialLine!.pointerEndLocation.dx;
+    final double sy = _initialLine!.pointerEndLocation.dy;
+
+    final double nfx = _currentLine!.pointerStartLocation.dx;
+    final double nfy = _currentLine!.pointerStartLocation.dy;
+    final double nsx = _currentLine!.pointerEndLocation.dx;
+    final double nsy = _currentLine!.pointerEndLocation.dy;
+
+    final double angle1 = math.atan2(fy - sy, fx - sx);
+    final double angle2 = math.atan2(nfy - nsy, nfx - nsx);
+
+    return angle2 - angle1;
+  }
+
+  @override
+  void addAllowedPointer(PointerDownEvent event) {
+    super.addAllowedPointer(event);
+    _velocityTrackers[event.pointer] = VelocityTracker.withKind(event.kind);
+    if (_state == _ScaleAndRotateState.ready) {
+      _state = _ScaleAndRotateState.possible;
+      _initialSpan = 0.0;
+      _currentSpan = 0.0;
+      _initialHorizontalSpan = 0.0;
+      _currentHorizontalSpan = 0.0;
+      _initialVerticalSpan = 0.0;
+      _currentVerticalSpan = 0.0;
+      _pointerLocations = <int, Offset>{};
+      _pointerQueue = <int>[];
+    }
+  }
+
+  @override
+  void handleEvent(PointerEvent event) {
+    assert(_state != _ScaleAndRotateState.ready);
+    bool didChangeConfiguration = false;
+    bool shouldStartIfAccepted = false;
+    if (event is PointerMoveEvent) {
+      final VelocityTracker tracker = _velocityTrackers[event.pointer]!;
+      if (!event.synthesized) {
+        tracker.addPosition(event.timeStamp, event.position);
+      }
+      _pointerLocations[event.pointer] = event.position;
+      shouldStartIfAccepted = true;
+      _lastTransform = event.transform;
+    } else if (event is PointerDownEvent) {
+      _pointerLocations[event.pointer] = event.position;
+      _pointerQueue.add(event.pointer);
+      didChangeConfiguration = true;
+      shouldStartIfAccepted = true;
+      _lastTransform = event.transform;
+    } else if (event is PointerUpEvent || event is PointerCancelEvent) {
+      _pointerLocations.remove(event.pointer);
+      _pointerQueue.remove(event.pointer);
+      didChangeConfiguration = true;
+      _lastTransform = event.transform;
+    }
+
+    _updateLines();
+    _update();
+
+    if (!didChangeConfiguration || _reconfigure(event.pointer)) {
+      _advanceStateMachine(shouldStartIfAccepted, event.kind);
+    }
+    stopTrackingIfPointerNoLongerDown(event);
+  }
+
+  void _update() {
+    final int count = _pointerLocations.keys.length;
+
+    final Offset? previousFocalPoint = _currentFocalPoint;
+
+    // Compute the focal point
+    Offset focalPoint = Offset.zero;
+    for (final int pointer in _pointerLocations.keys) {
+      focalPoint += _pointerLocations[pointer]!;
+    }
+    _currentFocalPoint = count > 0 ? focalPoint / count.toDouble() : Offset.zero;
+
+    if (previousFocalPoint == null) {
+      _localFocalPoint = PointerEvent.transformPosition(
+        _lastTransform,
+        _currentFocalPoint!,
+      );
+      _delta = Offset.zero;
+    } else {
+      final Offset localPreviousFocalPoint = _localFocalPoint;
+      _localFocalPoint = PointerEvent.transformPosition(
+        _lastTransform,
+        _currentFocalPoint!,
+      );
+      _delta = _localFocalPoint - localPreviousFocalPoint;
+    }
+
+    // Span is the average deviation from focal point. Horizontal and vertical
+    // spans are the average deviations from the focal point's horizontal and
+    // vertical coordinates, respectively.
+    double totalDeviation = 0.0;
+    double totalHorizontalDeviation = 0.0;
+    double totalVerticalDeviation = 0.0;
+    for (final int pointer in _pointerLocations.keys) {
+      totalDeviation += (_currentFocalPoint! - _pointerLocations[pointer]!).distance;
+      totalHorizontalDeviation += (_currentFocalPoint!.dx - _pointerLocations[pointer]!.dx).abs();
+      totalVerticalDeviation += (_currentFocalPoint!.dy - _pointerLocations[pointer]!.dy).abs();
+    }
+    _currentSpan = count > 0 ? totalDeviation / count : 0.0;
+    _currentHorizontalSpan = count > 0 ? totalHorizontalDeviation / count : 0.0;
+    _currentVerticalSpan = count > 0 ? totalVerticalDeviation / count : 0.0;
+  }
+
+  /// Updates [_initialLine] and [_currentLine] accordingly to the situation of
+  /// the registered pointers.
+  void _updateLines() {
+    final int count = _pointerLocations.keys.length;
+    assert(_pointerQueue.length >= count);
+
+    /// In case of just one pointer registered, reconfigure [_initialLine]
+    if (count < 2) {
+      _initialLine = _currentLine;
+    } else if (_initialLine != null &&
+        _initialLine!.pointerStartId == _pointerQueue[0] &&
+        _initialLine!.pointerEndId == _pointerQueue[1]) {
+      /// Rotation updated, set the [_currentLine]
+      _currentLine = _LineBetweenPointers(
+        pointerStartId: _pointerQueue[0],
+        pointerStartLocation: _pointerLocations[_pointerQueue[0]]!,
+        pointerEndId: _pointerQueue[1],
+        pointerEndLocation: _pointerLocations[_pointerQueue[1]]!,
+      );
+    } else {
+      /// A new rotation process is on the way, set the [_initialLine]
+      _initialLine = _LineBetweenPointers(
+        pointerStartId: _pointerQueue[0],
+        pointerStartLocation: _pointerLocations[_pointerQueue[0]]!,
+        pointerEndId: _pointerQueue[1],
+        pointerEndLocation: _pointerLocations[_pointerQueue[1]]!,
+      );
+      _currentLine = _initialLine;
+    }
+  }
+
+  bool _reconfigure(int pointer) {
+    _initialSpan = _currentSpan;
+    _initialLine = _currentLine;
+    _initialHorizontalSpan = _currentHorizontalSpan;
+    _initialVerticalSpan = _currentVerticalSpan;
+    if (_state == _ScaleAndRotateState.started) {
+      if (onEnd != null) {
+        final VelocityTracker tracker = _velocityTrackers[pointer]!;
+
+        Velocity velocity = tracker.getVelocity();
+        if (_isFlingGesture(velocity)) {
+          final Offset pixelsPerSecond = velocity.pixelsPerSecond;
+          if (pixelsPerSecond.distanceSquared > kMaxFlingVelocity * kMaxFlingVelocity) {
+            velocity = Velocity(pixelsPerSecond: (pixelsPerSecond / pixelsPerSecond.distance) * kMaxFlingVelocity);
+          }
+          invokeCallback<void>(
+              'onEnd', () => onEnd!(ScaleEndDetails(velocity: velocity, pointerCount: _pointerQueue.length)));
+        } else {
+          invokeCallback<void>('onEnd', () => onEnd!(ScaleEndDetails(pointerCount: _pointerQueue.length)));
+        }
+      }
+      _state = _ScaleAndRotateState.accepted;
+      return false;
+    }
+    return true;
+  }
+
+  void _advanceStateMachine(bool shouldStartIfAccepted, PointerDeviceKind pointerDeviceKind) {
+    if (_state == _ScaleAndRotateState.ready) {
+      _state = _ScaleAndRotateState.possible;
+    }
+
+    if (_state == _ScaleAndRotateState.possible) {
+      // Accept a gesture if there are at least 2 pointers involved. This is the main difference from Flutter's _ScaleState
+      if (_pointerLocations.keys.length >= 2) {
+        resolve(GestureDisposition.accepted);
+      }
+    } else if (_state.index >= _ScaleAndRotateState.accepted.index) {
+      resolve(GestureDisposition.accepted);
+    }
+
+    if (_state == _ScaleAndRotateState.accepted && shouldStartIfAccepted) {
+      _state = _ScaleAndRotateState.started;
+      _dispatchOnStartCallbackIfNeeded();
+    }
+
+    if (_state == _ScaleAndRotateState.started && onUpdate != null) {
+      invokeCallback<void>('onUpdate', () {
+        onUpdate!(ScaleUpdateDetails(
+          scale: _scaleFactor,
+          horizontalScale: _horizontalScaleFactor,
+          verticalScale: _verticalScaleFactor,
+          focalPoint: _currentFocalPoint!,
+          localFocalPoint: _localFocalPoint,
+          rotation: _computeRotationFactor(),
+          pointerCount: _pointerQueue.length,
+          focalPointDelta: _delta,
+        ));
+      });
+    }
+  }
+
+  void _dispatchOnStartCallbackIfNeeded() {
+    assert(_state == _ScaleAndRotateState.started);
+    if (onStart != null) {
+      invokeCallback<void>('onStart', () {
+        onStart!(ScaleStartDetails(
+          focalPoint: _currentFocalPoint!,
+          localFocalPoint: _localFocalPoint,
+          pointerCount: _pointerQueue.length,
+        ));
+      });
+    }
+  }
+
+  @override
+  void acceptGesture(int pointer) {
+    if (_state == _ScaleAndRotateState.possible) {
+      _state = _ScaleAndRotateState.started;
+      _dispatchOnStartCallbackIfNeeded();
+      if (dragStartBehavior == DragStartBehavior.start) {
+        _initialSpan = _currentSpan;
+        _initialLine = _currentLine;
+        _initialHorizontalSpan = _currentHorizontalSpan;
+        _initialVerticalSpan = _currentVerticalSpan;
+      }
+    }
+  }
+
+  @override
+  void rejectGesture(int pointer) {
+    stopTrackingPointer(pointer);
+  }
+
+  @override
+  void didStopTrackingLastPointer(int pointer) {
+    switch (_state) {
+      case _ScaleAndRotateState.possible:
+        resolve(GestureDisposition.rejected);
+        break;
+      case _ScaleAndRotateState.ready:
+        assert(false); // We should have not seen a pointer yet
+        break;
+      case _ScaleAndRotateState.accepted:
+        break;
+      case _ScaleAndRotateState.started:
+        assert(false); // We should be in the accepted state when user is done
+        break;
+    }
+    _state = _ScaleAndRotateState.ready;
+  }
+
+  @override
+  void dispose() {
+    _velocityTrackers.clear();
+    super.dispose();
+  }
+
+  @override
+  String get debugDescription => 'scale';
+}
diff --git a/Engine/Extras/VirtualProduction/EpicStageApp/lib/main.dart b/Engine/Extras/VirtualProduction/EpicStageApp/lib/main.dart
new file mode 100644
--- /dev/null
+++ b/Engine/Extras/VirtualProduction/EpicStageApp/lib/main.dart
@@ -0,0 +1,366 @@
+// Copyright Epic Games, Inc. All Rights Reserved.
+
+import 'dart:async';
+import 'dart:isolate';
+import 'dart:ui';
+
+import 'package:flutter/material.dart';
+import 'package:flutter/services.dart';
+import 'package:flutter_gen/gen_l10n/app_localizations.dart';
+import 'package:logging/logging.dart';
+import 'package:provider/provider.dart';
+import 'package:streaming_shared_preferences/streaming_shared_preferences.dart';
+
+import 'models/engine_connection.dart';
+import 'models/engine_passphrase_manager.dart';
+import 'models/navigator_keys.dart';
+import 'models/preview_render_manager.dart';
+import 'models/settings/delta_widget_settings.dart';
+import 'models/settings/main_screen_settings.dart';
+import 'models/settings/recent_actor_settings.dart';
+import 'models/settings/selected_actor_settings.dart';
+import 'models/settings/stage_map_settings.dart';
+import 'models/unreal_actor_creator.dart';
+import 'models/unreal_actor_manager.dart';
+import 'models/unreal_property_manager.dart';
+import 'models/unreal_transaction_manager.dart';
+import 'routes.dart';
+import 'utilities/constants.dart';
+import 'utilities/logging.dart';
+import 'utilities/preferences_bundle.dart';
+import 'utilities/transient_preference.dart';
+import 'utilities/unreal_colors.dart';
+import 'widgets/elements/epic_scroll_view.dart';
+import 'widgets/screens/connect/connect.dart';
+import 'widgets/screens/eula/eula_screen.dart';
+
+final _log = Logger('Main');
+
+void main() async {
+  Logging.instance.initialize();
+
+  final asyncLog = Logger('Async');
+
+  // Catch any asynchronous errors
+  runZonedGuarded(
+    () => initAndRunApp(),
+    (error, stack) {
+      asyncLog.severe(error, error, stack);
+    },
+  );
+}
+
+/// Set up the app and run it. Encapsulated so we can wrap it in [runZonedGuarded] in [main] for better error handling.
+void initAndRunApp() async {
+  WidgetsFlutterBinding.ensureInitialized();
+
+  final flutterLog = Logger('Flutter');
+  final isolateLog = Logger('Isolate');
+
+  preloadShaders();
+
+  // Catch any errors originating from or otherwise passed on by Flutter
+  FlutterError.onError = (FlutterErrorDetails details) {
+    flutterLog.severe(details.exception, details.exception, details.stack);
+  };
+
+  // Catch errors outside the Flutter context
+  Isolate.current.addErrorListener(RawReceivePort((errorStackPair) {
+    isolateLog.severe(errorStackPair.first, errorStackPair.first, StackTrace.fromString(errorStackPair.second));
+  }).sendPort);
+
+  // Load user settings
+  final preferences = await StreamingSharedPreferences.instance;
+
+  // Run the app itself
+  runApp(EpicStageApp(preferences: preferences));
+}
+
+/// Preload any shaders we expect to use regularly.
+void preloadShaders() async {
+  await EpicScrollView.preloadShaders();
+}
+
+/// Unload any preloaded shaders.
+void unloadShaders() {
+  EpicScrollView.unloadShaders();
+}
+
+class EpicStageApp extends StatefulWidget {
+  const EpicStageApp({Key? key, required this.preferences}) : super(key: key);
+
+  final StreamingSharedPreferences preferences;
+
+  @override
+  _EpicStageAppState createState() => _EpicStageAppState();
+}
+
+class _EpicStageAppState extends State<EpicStageApp> with TickerProviderStateMixin, WidgetsBindingObserver {
+  late final PreferencesBundle _preferenceBundle;
+
+  /// Whether or not the user has accepted the EULA agreement.
+  bool bHasAcceptedEula = false;
+
+  @override
+  void initState() {
+    super.initState();
+
+    _log.info('App initialized');
+
+    _preferenceBundle = PreferencesBundle(widget.preferences, TransientSharedPreferences());
+
+    bHasAcceptedEula = widget.preferences.getBool('common.bHasAcceptedEula', defaultValue: false).getValue();
+
+    WidgetsBinding.instance.addObserver(this);
+  }
+
+  @override
+  void dispose() {
+    WidgetsBinding.instance.removeObserver(this);
+
+    _log.info('App disposed');
+
+    super.dispose();
+  }
+
+  @override
+  void didChangeAppLifecycleState(AppLifecycleState state) {
+    super.didChangeAppLifecycleState(state);
+
+    switch (state) {
+      case AppLifecycleState.detached:
+        _log.info('App detached');
+        break;
+
+      case AppLifecycleState.inactive:
+        _log.info('App inactive');
+        break;
+
+      case AppLifecycleState.paused:
+        _log.info('App paused');
+        break;
+
+      case AppLifecycleState.resumed:
+        _log.info('App resumed');
+        break;
+    }
+  }
+
+  @override
+  Widget build(BuildContext context) {
+    // Force landscape layout on mobile devices
+    SystemChrome.setPreferredOrientations([
+      DeviceOrientation.landscapeLeft,
+      DeviceOrientation.landscapeRight,
+    ]);
+
+    var routes = {for (var entry in RouteData.allRoutes.entries) entry.key: entry.value.createScreen};
+
+    if (!bHasAcceptedEula) {
+      routes['/'] = routes[EulaScreen.route]!;
+    } else {
+      routes['/'] = routes[ConnectScreen.route]!;
+    }
+
+    const colorScheme = ColorScheme.dark(
+      primary: UnrealColors.highlightBlue,
+      secondary: Color(0xff575757),
+      onPrimary: UnrealColors.white,
+      onSecondary: UnrealColors.white,
+      background: UnrealColors.gray06,
+      surfaceVariant: UnrealColors.gray10,
+      surface: UnrealColors.gray14,
+      surfaceTint: UnrealColors.gray18,
+      onSurface: UnrealColors.gray75,
+      shadow: Colors.transparent,
+    );
+
+    const textTheme = const TextTheme(
+      bodyMedium: TextStyle(
+        color: UnrealColors.gray75,
+        decorationColor: UnrealColors.gray75,
+        fontVariations: [FontVariation('wght', 400)],
+        fontSize: 14,
+      ),
+      displayLarge: TextStyle(
+        color: UnrealColors.white,
+        decorationColor: UnrealColors.white,
+        fontVariations: [FontVariation('wght', 700)],
+        fontSize: 14,
+      ),
+      displayMedium: TextStyle(
+        color: UnrealColors.white,
+        decorationColor: UnrealColors.white,
+        fontVariations: [FontVariation('wght', 400)],
+        fontSize: 14,
+      ),
+      headlineSmall: TextStyle(
+        color: UnrealColors.gray75,
+        decorationColor: UnrealColors.gray75,
+        fontVariations: [FontVariation('wght', 600)],
+        letterSpacing: 0.25,
+        fontSize: 14,
+      ),
+      titleLarge: TextStyle(
+        color: UnrealColors.gray90,
+        decorationColor: UnrealColors.gray90,
+        fontVariations: [FontVariation('wght', 700)],
+        letterSpacing: 1,
+        fontSize: 12,
+      ),
+      labelMedium: TextStyle(
+        color: UnrealColors.gray56,
+        decorationColor: UnrealColors.gray56,
+        fontVariations: [FontVariation('wght', 400)],
+        letterSpacing: 0.5,
+        fontSize: 12,
+      ),
+    );
+
+    return MultiProvider(
+      providers: [
+        Provider<EnginePassphraseManager>(create: (_) => EnginePassphraseManager()),
+        Provider<PreferencesBundle>(create: (_) => _preferenceBundle),
+        Provider<SelectedActorSettings>(create: (_) => SelectedActorSettings(_preferenceBundle)),
+        Provider<StageMapSettings>(create: (_) => StageMapSettings(_preferenceBundle)),
+        Provider<RecentActorSettings>(create: (_) => RecentActorSettings(_preferenceBundle)),
+        Provider<MainScreenSettings>(create: (_) => MainScreenSettings(_preferenceBundle)),
+        Provider<DeltaWidgetSettings>(create: (_) => DeltaWidgetSettings(_preferenceBundle)),
+        Provider<EngineConnectionManager>(
+          create: (context) => EngineConnectionManager(context),
+          dispose: (_, value) => value.dispose(),
+        ),
+        Provider<UnrealTransactionManager>(create: (context) => UnrealTransactionManager(context)),
+        Provider<UnrealPropertyManager>(create: (context) => UnrealPropertyManager(context)),
+        Provider<UnrealActorManager>(create: (context) => UnrealActorManager(context)),
+        Provider<PreviewRenderManager>(
+          create: (context) => PreviewRenderManager(context),
+          dispose: (context, value) => value.dispose(),
+        ),
+        Provider<UnrealActorCreator>(
+          create: (context) => UnrealActorCreator(context),
+          dispose: (context, value) => value.dispose(),
+        ),
+      ],
+      child: MaterialApp(
+        theme: ThemeData(
+          fontFamily: 'Inter',
+
+          textTheme: textTheme,
+          colorScheme: colorScheme,
+
+          // Disable Material "ink" splash effects
+          splashFactory: NoSplash.splashFactory,
+
+          scaffoldBackgroundColor: UnrealColors.gray10,
+          hoverColor: Colors.transparent,
+          highlightColor: Colors.transparent,
+
+          textButtonTheme: TextButtonThemeData(
+            style: ButtonStyle(
+              overlayColor: MaterialStateProperty.all(Colors.transparent),
+            ),
+          ),
+
+          checkboxTheme: CheckboxThemeData(
+            overlayColor: MaterialStateProperty.all(Colors.transparent),
+          ),
+
+          appBarTheme: const AppBarTheme(
+            shadowColor: Colors.transparent,
+          ),
+
+          tooltipTheme: const TooltipThemeData(
+            waitDuration: Duration(milliseconds: 700),
+            decoration: BoxDecoration(
+              image: DecorationImage(
+                image: AssetImage('assets/images/decoration/tooltip.png'),
+                centerSlice: Rect.fromLTRB(4, 4, 60, 60),
+              ),
+            ),
+          ),
+
+          textSelectionTheme: const TextSelectionThemeData(
+            cursorColor: UnrealColors.white,
+          ),
+
+          inputDecorationTheme: InputDecorationTheme(
+            filled: true,
+            outlineBorder: BorderSide.none,
+            border: OutlineInputBorder(
+              borderRadius: BorderRadius.circular(4),
+              borderSide: BorderSide.none,
+            ),
+            enabledBorder: OutlineInputBorder(
+              borderRadius: BorderRadius.circular(4),
+              borderSide: BorderSide.none,
+            ),
+            focusedBorder: OutlineInputBorder(
+              borderRadius: BorderRadius.circular(4),
+              borderSide: BorderSide.none,
+            ),
+            errorBorder: OutlineInputBorder(
+              borderRadius: BorderRadius.circular(4),
+              borderSide: BorderSide.none,
+            ),
+            fillColor: colorScheme.background,
+            hoverColor: colorScheme.background,
+            contentPadding: const EdgeInsets.symmetric(
+              horizontal: 12,
+            ),
+            hintStyle: textTheme.bodyMedium!.copyWith(color: UnrealColors.gray56),
+          ),
+
+          listTileTheme: const ListTileThemeData(
+            textColor: UnrealColors.gray75,
+            iconColor: UnrealColors.gray75,
+            selectedColor: UnrealColors.white,
+            selectedTileColor: UnrealColors.highlightBlue,
+          ),
+
+          scrollbarTheme: const ScrollbarThemeData(
+            thumbColor: MaterialStatePropertyAll(UnrealColors.gray22),
+            thickness: MaterialStatePropertyAll(8),
+            radius: Radius.circular(4),
+            mainAxisMargin: cardCornerRadius,
+            crossAxisMargin: cardCornerRadius,
+          ),
+
+          cardTheme: CardTheme(
+            color: UnrealColors.gray10,
+            clipBehavior: Clip.antiAlias,
+            shape: RoundedRectangleBorder(borderRadius: BorderRadius.circular(cardCornerRadius)),
+            margin: EdgeInsets.zero,
+            shadowColor: Colors.transparent,
+          ),
+        ),
+        localizationsDelegates: AppLocalizations.localizationsDelegates,
+        supportedLocales: AppLocalizations.supportedLocales,
+        navigatorKey: rootNavigatorKey,
+        onGenerateTitle: (context) => AppLocalizations.of(context)!.appTitle,
+        onGenerateRoute: (RouteSettings settings) {
+          Widget Function(BuildContext)? pageFunction = routes[settings.name];
+          if (pageFunction != null) {
+            return PageRouteBuilder(
+              settings: settings,
+              // Add safe area on all pages to avoid overlapping the status bar
+              pageBuilder: (context, animation, secondaryAnimation) => Container(
+                color: Theme.of(context).colorScheme.background,
+                child: SafeArea(
+                  left: false,
+                  right: false,
+                  bottom: false,
+                  top: true,
+                  child: pageFunction(context),
+                ),
+              ),
+              transitionDuration: Duration.zero,
+            );
+          }
+
+          return MaterialPageRoute(builder: (_) => Text("Invalid page ${settings.name ?? "(null)"}"));
+        },
+      ),
+    );
+  }
+}
diff --git a/Engine/Extras/VirtualProduction/EpicStageApp/lib/models/actor_data/generic_actor_data.dart b/Engine/Extras/VirtualProduction/EpicStageApp/lib/models/actor_data/generic_actor_data.dart
new file mode 100644
--- /dev/null
+++ b/Engine/Extras/VirtualProduction/EpicStageApp/lib/models/actor_data/generic_actor_data.dart
@@ -0,0 +1,48 @@
+// Copyright Epic Games, Inc. All Rights Reserved.
+
+import '../property_modify_operations.dart';
+import '../unreal_property_manager.dart';
+import '../unreal_types.dart';
+import 'per_class_actor_data.dart';
+
+/// Retrieves generically useful data for actor in Unreal Engine.
+class GenericActorData extends UnrealPerClassActorData {
+  GenericActorData(super.params);
+
+  late final TrackedPropertyId _isHiddenInGamePropertyId;
+
+  /// Whether this actor is hidden in-game.
+  bool get bIsHiddenInGame => propertyManager.getTrackedPropertyValue(_isHiddenInGamePropertyId) ?? false;
+
+  set bIsHiddenInGame(bool bNewValue) {
+    if (!bIsDataReady) {
+      return;
+    }
+
+    if (propertyManager.beginTransaction('Set Hidden In Game')) {
+      propertyManager.modifyTrackedPropertyValue(
+        _isHiddenInGamePropertyId,
+        const SetOperation(),
+        deltaValue: bNewValue,
+      );
+      propertyManager.endTransaction();
+    }
+  }
+
+  @override
+  Future<void> initialize() async {
+    _isHiddenInGamePropertyId = propertyManager.trackProperty(
+      UnrealProperty(objectPath: actor.path, propertyName: 'bHidden'),
+      onPropertyChanged,
+    );
+
+    await propertyManager.waitForProperty(_isHiddenInGamePropertyId);
+  }
+
+  @override
+  void dispose() {
+    propertyManager.stopTrackingProperty(_isHiddenInGamePropertyId, onPropertyChanged);
+
+    super.dispose();
+  }
+}
diff --git a/Engine/Extras/VirtualProduction/EpicStageApp/lib/models/actor_data/light_card_actor_data.dart b/Engine/Extras/VirtualProduction/EpicStageApp/lib/models/actor_data/light_card_actor_data.dart
new file mode 100644
--- /dev/null
+++ b/Engine/Extras/VirtualProduction/EpicStageApp/lib/models/actor_data/light_card_actor_data.dart
@@ -0,0 +1,67 @@
+// Copyright Epic Games, Inc. All Rights Reserved.
+
+import '../engine_connection.dart';
+import '../unreal_property_manager.dart';
+import '../unreal_types.dart';
+import 'per_class_actor_data.dart';
+
+/// Retrieves [LightCardActorData] for an actor from the engine.
+class LightCardActorData extends UnrealPerClassActorData {
+  LightCardActorData(super.params);
+
+  late final TrackedPropertyId _isUVPropertyId;
+  late final TrackedPropertyId _isFlagPropertyId;
+
+  /// If this is set, the UV property can't be tracked and we should use this value instead.
+  bool? _bIsUVFallback;
+
+  /// Whether this is light card is positioned in UV space.
+  bool get bIsUV => _bIsUVFallback ?? propertyManager.getTrackedPropertyValue(_isUVPropertyId) ?? false;
+
+  /// Whether this light card is a flag.
+  bool get bIsFlag => propertyManager.getTrackedPropertyValue(_isFlagPropertyId) ?? false;
+
+  @override
+  Future<void> initialize() async {
+    // Old API version can't subscribe to these properties
+    if (!connectionManager.apiVersion!.bSupportsLightcardSubtypes) {
+      _isUVPropertyId = TrackedPropertyId.invalid();
+      _isFlagPropertyId = TrackedPropertyId.invalid();
+
+      // This property can't be subscribed to on this API version, but we can at least get its value directly and assume
+      // that it won't change
+      final UnrealHttpResponse response = await connectionManager.sendHttpRequest(UnrealHttpRequest(
+        url: '/remote/object/property',
+        verb: 'PUT',
+        body: {
+          'objectPath': actor.path,
+          'propertyName': 'bIsUVLightCard',
+          'access': 'READ_ACCESS',
+        },
+      ));
+
+      _bIsUVFallback = response.body?['bIsUVLightCard'];
+
+      return Future.value();
+    }
+
+    _isUVPropertyId = propertyManager.trackProperty(
+        UnrealProperty(objectPath: actor.path, propertyName: 'bIsUVLightCard'), onPropertyChanged);
+
+    _isFlagPropertyId = propertyManager.trackProperty(
+        UnrealProperty(objectPath: actor.path, propertyName: 'bIsLightCardFlag'), onPropertyChanged);
+
+    await Future.wait([
+      propertyManager.waitForProperty(_isUVPropertyId),
+      propertyManager.waitForProperty(_isFlagPropertyId),
+    ]);
+  }
+
+  @override
+  void dispose() {
+    propertyManager.stopTrackingProperty(_isUVPropertyId, onPropertyChanged);
+    propertyManager.stopTrackingProperty(_isFlagPropertyId, onPropertyChanged);
+
+    super.dispose();
+  }
+}
diff --git a/Engine/Extras/VirtualProduction/EpicStageApp/lib/models/actor_data/per_class_actor_data.dart b/Engine/Extras/VirtualProduction/EpicStageApp/lib/models/actor_data/per_class_actor_data.dart
new file mode 100644
--- /dev/null
+++ b/Engine/Extras/VirtualProduction/EpicStageApp/lib/models/actor_data/per_class_actor_data.dart
@@ -0,0 +1,74 @@
+// Copyright Epic Games, Inc. All Rights Reserved.
+
+import 'package:flutter/material.dart';
+
+import '../engine_connection.dart';
+import '../unreal_property_manager.dart';
+import '../unreal_types.dart';
+
+/// Parameters for creating an [UnrealPerClassActorData].
+class UnrealPerClassActorDataParams {
+  const UnrealPerClassActorDataParams({
+    required this.actor,
+    required this.connectionManager,
+    required this.propertyManager,
+  });
+
+  /// The actor to which this data corresponds.
+  final UnrealObject actor;
+
+  /// The manager responsible for the engine connection.
+  final EngineConnectionManager connectionManager;
+
+  /// The property manager used to retrieve the data.
+  final UnrealPropertyManager propertyManager;
+}
+
+/// Base class for extra data to retrieve about an actor from Unreal Engine based on the actor's class.
+abstract class UnrealPerClassActorData extends ChangeNotifier {
+  UnrealPerClassActorData(UnrealPerClassActorDataParams params)
+      : _actor = params.actor,
+        _propertyManager = params.propertyManager,
+        _connectionManager = params.connectionManager {
+    initialize().then((_) {
+      _bIsDataReady = true;
+      notifyListeners();
+    });
+  }
+
+  /// The actor to which this data corresponds.
+  final UnrealObject _actor;
+
+  /// The property manager used to retrieve the data.
+  final UnrealPropertyManager _propertyManager;
+
+  /// The manager responsible for the engine connection.
+  final EngineConnectionManager _connectionManager;
+
+  /// Whether this fully retrieved its initial data from the engine.
+  bool _bIsDataReady = false;
+
+  /// Whether this fully retrieved its initial data from the engine.
+  bool get bIsDataReady => _bIsDataReady;
+
+  /// Async function to retrieve any data from the engine.
+  @protected
+  Future<void> initialize();
+
+  /// The property manager used to retrieve the data.
+  @protected
+  UnrealPropertyManager get propertyManager => _propertyManager;
+
+  /// The connection manager used to retrieve the data.
+  @protected
+  EngineConnectionManager get connectionManager => _connectionManager;
+
+  /// The actor to which this data corresponds.
+  @protected
+  UnrealObject get actor => _actor;
+
+  /// Convenience function which accepts a property argument from callbacks for UnrealPropertyManager.trackProperty.
+  void onPropertyChanged(_) {
+    notifyListeners();
+  }
+}
diff --git a/Engine/Extras/VirtualProduction/EpicStageApp/lib/models/api_version.dart b/Engine/Extras/VirtualProduction/EpicStageApp/lib/models/api_version.dart
new file mode 100644
--- /dev/null
+++ b/Engine/Extras/VirtualProduction/EpicStageApp/lib/models/api_version.dart
@@ -0,0 +1,81 @@
+// Copyright Epic Games, Inc. All Rights Reserved.
+
+/// A semantic version indicating which engine API is available to the app.
+class EpicStageAppAPIVersion implements Comparable<EpicStageAppAPIVersion> {
+  const EpicStageAppAPIVersion(this.major, this.minor, this.patch);
+
+  factory EpicStageAppAPIVersion.fromString(String string) {
+    final List<String> components = string.split('.');
+    if (components.length != 3) {
+      throw ArgumentError('Expected 3 components in version number but got ${components.length}');
+    }
+
+    return EpicStageAppAPIVersion(
+      int.parse(components[0]),
+      int.parse(components[1]),
+      int.parse(components[2]),
+    );
+  }
+
+  final int major;
+  final int minor;
+  final int patch;
+
+  @override
+  int compareTo(EpicStageAppAPIVersion other) {
+    int comparison = major.compareTo(other.major);
+    if (comparison != 0) {
+      return comparison;
+    }
+
+    comparison = minor.compareTo(other.minor);
+    if (comparison != 0) {
+      return comparison;
+    }
+
+    return patch.compareTo(other.patch);
+  }
+
+  @override
+  bool operator ==(Object other) {
+    if (!(other is EpicStageAppAPIVersion)) {
+      return false;
+    }
+
+    return compareTo(other) == 0;
+  }
+
+  @override
+  int get hashCode => Object.hash(major, minor, patch);
+
+  bool operator >(EpicStageAppAPIVersion other) => compareTo(other) > 0;
+  bool operator >=(EpicStageAppAPIVersion other) => compareTo(other) >= 0;
+  bool operator <(EpicStageAppAPIVersion other) => compareTo(other) <= 0;
+  bool operator <=(EpicStageAppAPIVersion other) => compareTo(other) <= 0;
+
+  @override
+  String toString() {
+    return '$major.$minor.$patch';
+  }
+
+  /// Whether the engine API has a dedicated GetRemoteControlWebInterfacePort function.
+  bool get bHasWebInterfacePortFunction => this >= const EpicStageAppAPIVersion(1, 1, 0);
+
+  /// Whether the engine API has full support for the bIsLightCardFlag and bIsUVLightCard properties;
+  bool get bSupportsLightcardSubtypes => this >= const EpicStageAppAPIVersion(1, 3, 0);
+
+  /// Whether the engine API supports append/insert/remove for array properties.
+  bool get bSupportsArrayPropertyOperations => this >= const EpicStageAppAPIVersion(1, 4, 0);
+
+  /// Whether the engine API supports resetting struct properties contained within lists.
+  bool get bCanResetListItems => this >= const EpicStageAppAPIVersion(1, 5, 0);
+
+  /// Whether the engine API correctly handles HTTP property set calls as part of a manual transaction.
+  bool get bCanHttpSetPropertyInManualTransaction => this >= const EpicStageAppAPIVersion(1, 6, 0);
+
+  /// Whether the engine API allows for query parameters in URLs passed to the HTTP WebSocket route.
+  bool get bCanUseQueryParamsInWebSocketHttpUrl => this >= const EpicStageAppAPIVersion(1, 7, 0);
+
+  /// Whether the initial actor position passed to the ndisplay.preview.actor.create route is applied correctly.
+  bool get bIsNewActorOverridePositionAccurate => this >= const EpicStageAppAPIVersion(1, 8, 0);
+}
diff --git a/Engine/Extras/VirtualProduction/EpicStageApp/lib/models/color.dart b/Engine/Extras/VirtualProduction/EpicStageApp/lib/models/color.dart
new file mode 100644
--- /dev/null
+++ b/Engine/Extras/VirtualProduction/EpicStageApp/lib/models/color.dart
@@ -0,0 +1,245 @@
+// Copyright Epic Games, Inc. All Rights Reserved.
+
+import 'dart:math';
+
+import 'package:flutter/cupertino.dart';
+import 'package:vector_math/vector_math.dart';
+
+/// A color represented as floating-point values as opposed to Dart's native, int-based [Color].
+@immutable
+class FloatColor {
+  const FloatColor(this.red, this.green, this.blue, [this.alpha = 1]);
+
+  final double red;
+  final double green;
+  final double blue;
+  final double alpha;
+
+  static const FloatColor zero = FloatColor(0, 0, 0, 0);
+
+  /// Create a color from an existing HSVColor.
+  static FloatColor fromHSVColor(NegatableHSVColor hsvColor) {
+    return fromHSV(hsvColor.hue, hsvColor.saturation, hsvColor.value, hsvColor.alpha);
+  }
+
+  /// Create a color from hue, saturation, and value.
+  static FloatColor fromHSV(double hue, double saturation, double value, [double alpha = 1]) {
+    final double hDiv60 = hue / 60;
+    final double hDiv60Fraction = hDiv60 - hDiv60.floorToDouble();
+
+    final List<double> rgbValues = [
+      value,
+      value * (1 - saturation),
+      value * (1 - (hDiv60Fraction * saturation)),
+      value * (1 - ((1 - hDiv60Fraction) * saturation)),
+    ];
+
+    const List<List<int>> rgbSwizzle = [
+      [0, 3, 1],
+      [2, 0, 1],
+      [1, 0, 3],
+      [1, 2, 0],
+      [3, 1, 0],
+      [0, 1, 2],
+    ];
+
+    final int swizzleIndex = hDiv60.floor() % 6;
+
+    return FloatColor(
+      rgbValues[rgbSwizzle[swizzleIndex][0]],
+      rgbValues[rgbSwizzle[swizzleIndex][1]],
+      rgbValues[rgbSwizzle[swizzleIndex][2]],
+      alpha,
+    );
+  }
+
+  /// Convert to a Dart native color.
+  Color toColor() {
+    return Color.fromARGB(
+      _doubleValueToInt(alpha),
+      _doubleValueToInt(red),
+      _doubleValueToInt(green),
+      _doubleValueToInt(blue),
+    );
+  }
+
+  /// Copy the color, overriding its [alpha] value.
+  FloatColor withAlpha(double newAlpha) {
+    return FloatColor(red, green, blue, newAlpha);
+  }
+
+  /// Piecewise addition of each color component.
+  FloatColor operator +(FloatColor other) {
+    return FloatColor(red + other.red, green + other.green, blue + other.blue, alpha + other.alpha);
+  }
+
+  /// Piecewise subtraction of each color component.
+  FloatColor operator -(FloatColor other) {
+    return FloatColor(red - other.red, green - other.green, blue - other.blue, alpha - other.alpha);
+  }
+
+  /// Piecewise multiplication of each color component.
+  FloatColor operator *(double multiplier) {
+    return FloatColor(red * multiplier, green * multiplier, blue * multiplier, alpha * multiplier);
+  }
+
+  /// Piecewise division of each color component.
+  FloatColor operator /(double divisor) {
+    return FloatColor(red / divisor, green / divisor, blue / divisor, alpha / divisor);
+  }
+
+  /// Convert this to a JSON representation compatible with Unreal Engine.
+  Map<String, dynamic> toJson() {
+    return {'R': red, 'G': green, 'B': blue, 'A': alpha};
+  }
+
+  /// Convert this to a user-friendly string representation.
+  @override
+  String toString() {
+    return '(${red.toStringAsFixed(2)},'
+        '${green.toStringAsFixed(2)},'
+        '${blue.toStringAsFixed(2)},'
+        '${alpha.toStringAsFixed(2)})';
+  }
+
+  /// Convert a floating-point color value to an 8-bit integer.
+  int _doubleValueToInt(double value) {
+    return (value * 255).round().clamp(0, 255);
+  }
+}
+
+/// An HSV color that imposes no limits on the value or alpha components, allowing for negative/subtractive colors.
+@immutable
+class NegatableHSVColor {
+  /// Creates a color.
+  const NegatableHSVColor.fromAHSV(this.alpha, this.hue, this.saturation, this.value)
+      : assert(hue >= 0.0),
+        assert(hue <= 360.0),
+        assert(saturation >= 0.0),
+        assert(saturation <= 1.0);
+
+  /// Creates a [NegatableHSVColor] from an RGB [FloatColor].
+  factory NegatableHSVColor.fromColor(FloatColor color) {
+    final List<double> rgbList = [color.red, color.green, color.blue];
+    final double rgbMin = rgbList.reduce(min);
+    final double rgbMax = rgbList.reduce(max);
+    final double rgbRange = rgbMax - rgbMin;
+
+    double hue;
+    if (rgbMax == rgbMin) {
+      hue = 0;
+    } else if (rgbMax == color.red) {
+      hue = ((((color.green - color.blue) / rgbRange) * 60) + 360) % 360;
+    } else if (rgbMax == color.green) {
+      hue = (((color.blue - color.red) / rgbRange) * 60) + 120;
+    } else {
+      // rgbMax == blue
+      hue = (((color.red - color.green) / rgbRange) * 60) + 240;
+    }
+
+    final double saturation = (rgbMax == 0) ? 0 : (rgbRange / rgbMax).clamp(0, 1);
+    final double value = rgbMax;
+
+    return NegatableHSVColor.fromAHSV(color.alpha, hue, saturation, value);
+  }
+
+  /// Copy the color, overriding its [alpha] value.
+  NegatableHSVColor withAlpha(double newAlpha) {
+    return NegatableHSVColor.fromAHSV(newAlpha, hue, saturation, value);
+  }
+
+  /// Copy the color, overriding its [hue] value.
+  NegatableHSVColor withHue(double newHue) {
+    return NegatableHSVColor.fromAHSV(alpha, newHue, saturation, value);
+  }
+
+  /// Copy the color, overriding its [saturation] value.
+  NegatableHSVColor withSaturation(double newSaturation) {
+    return NegatableHSVColor.fromAHSV(alpha, hue, newSaturation, value);
+  }
+
+  /// Copy the color, overriding its [value] value.
+  NegatableHSVColor withValue(double newValue) {
+    return NegatableHSVColor.fromAHSV(alpha, hue, saturation, newValue);
+  }
+
+  final double alpha;
+  final double hue;
+  final double saturation;
+  final double value;
+}
+
+/// A color represented by its position on an HSV color wheel. We use this representation internally since it reduces
+/// imprecision caused by RGB (ambiguous hue/saturation at low value) or HSV (huge hue delta values when the position
+/// is near the center of the wheel, i.e. low saturation) representations.
+@immutable
+class WheelColor {
+  /// Create a [WheelColor] from its individual components.
+  const WheelColor(this.position, this.value, this.alpha);
+
+  /// Create a [WheelColor] from hue/saturation/value.
+  factory WheelColor.fromHSV(double hue, double saturation, double value, double alpha) {
+    final Offset position = Offset.fromDirection(radians(hue), saturation);
+    return WheelColor(position, value, alpha);
+  }
+
+  /// Create a [WheelColor] from an [HSVColor].
+  factory WheelColor.fromHSVColor(NegatableHSVColor color) {
+    return WheelColor.fromHSV(color.hue, color.saturation, color.value, color.alpha);
+  }
+
+  /// Create a [WheelColor] from a [FloatColor].
+  factory WheelColor.fromFloatColor(FloatColor color) {
+    return WheelColor.fromHSVColor(NegatableHSVColor.fromColor(color));
+  }
+
+  /// A [WheelColor] with all values set to 0.
+  static const WheelColor zero = WheelColor(Offset.zero, 0, 0);
+
+  /// Position of the color on a color wheel of radius 1.
+  final Offset position;
+
+  /// The value component of the color's HSV value.
+  final double value;
+
+  /// The alpha component of the color.
+  final double alpha;
+
+  /// Get the hue of the color this represents.
+  double get hue => (degrees(position.direction) + 360.0) % 360.0;
+
+  /// Get the saturation of the color this represents.
+  double get saturation => position.distance;
+
+  /// Piecewise addition of each color component.
+  WheelColor operator +(WheelColor other) =>
+      WheelColor(position + other.position, value + other.value, alpha + other.alpha);
+
+  /// Piecewise subtraction of each color component.
+  WheelColor operator -(WheelColor other) =>
+      WheelColor(position - other.position, value - other.value, alpha - other.alpha);
+
+  /// Convert this to an HSV color.
+  NegatableHSVColor toHSVColor() => NegatableHSVColor.fromAHSV(
+        alpha,
+        ((hue % 360) + 360) % 360,
+        saturation.clamp(0, 1),
+        value,
+      );
+
+  /// Convert this to a floating-point RGB color.
+  FloatColor toFloatColor() => FloatColor.fromHSVColor(toHSVColor());
+
+  /// Convert this to a JSON representation compatible with Unreal Engine.
+  /// If [bUseLuminance] is true, replace Alpha with Luminance.
+  Map<String, dynamic> toJson({bool bUseLuminance = false}) {
+    return {
+      'Position': {
+        'X': position.dx,
+        'Y': position.dy,
+      },
+      'Value': value,
+      (bUseLuminance ? 'Luminance' : 'Alpha'): alpha,
+    };
+  }
+}
diff --git a/Engine/Extras/VirtualProduction/EpicStageApp/lib/models/engine_connection.dart b/Engine/Extras/VirtualProduction/EpicStageApp/lib/models/engine_connection.dart
new file mode 100644
--- /dev/null
+++ b/Engine/Extras/VirtualProduction/EpicStageApp/lib/models/engine_connection.dart
@@ -0,0 +1,845 @@
+// Copyright Epic Games, Inc. All Rights Reserved.
+
+import 'dart:async';
+import 'dart:convert';
+import 'dart:io';
+
+import 'package:crypto/crypto.dart';
+import 'package:flutter/cupertino.dart';
+import 'package:flutter_gen/gen_l10n/app_localizations.dart';
+import 'package:logging/logging.dart';
+import 'package:provider/provider.dart';
+import 'package:shared_preferences/shared_preferences.dart';
+import 'package:uuid/uuid.dart';
+import 'package:web_socket_channel/io.dart';
+
+import './navigator_keys.dart';
+import '../utilities/net_utilities.dart';
+import '../utilities/unreal_utilities.dart';
+import '../widgets/elements/modal.dart';
+import '../widgets/screens/connect/connect.dart';
+import '../widgets/screens/connect/views/passphrase_modal.dart';
+import '../widgets/screens/main/stage_app_main_screen.dart';
+import '../widgets/screens/reconnect_screen.dart';
+import 'api_version.dart';
+import 'engine_passphrase_manager.dart';
+
+final _log = Logger('EngineConnection');
+
+/// The state of a connection to Unreal Engine.
+enum EngineConnectionState {
+  connected,
+  disconnected,
+}
+
+/// Possible results of attempting to connect to the engine.
+enum EngineConnectionResult {
+  /// The connection was successful.
+  success,
+
+  /// The connection was cancelled.
+  cancelled,
+
+  /// The connection couldn't be opened because another one was already open.
+  alreadyOpen,
+
+  /// The connection failed because the passphrase was incorrect.
+  passphraseRejected,
+
+  /// The connection failed for unknown reasons.
+  genericFailure,
+}
+
+/// HTTP response codes
+class HttpResponseCode {
+  static const int ok = 200;
+  static const int unauthorized = 401;
+  static const int clientClosedRequest = 499;
+}
+
+/// Callback type for a WebSocket message. Passes dynamic JSON data parsed from the message.
+typedef WebSocketMessageCallback = void Function(dynamic message);
+
+/// Callback for changes to the state of the engine connection.
+typedef ConnectionStateCallback = void Function(EngineConnectionState connectionState);
+
+/// An HTTP request in the format required by Unreal Engine.
+class UnrealHttpRequest {
+  const UnrealHttpRequest({required this.url, required this.verb, this.body});
+
+  final String url;
+  final String verb;
+  final dynamic body;
+}
+
+/// An HTTP response received from Unreal Engine.
+class UnrealHttpResponse {
+  const UnrealHttpResponse({required this.requestId, required this.code, required this.body});
+
+  final int requestId;
+  final int code;
+  final dynamic body;
+}
+
+/// An HTTP request to the engine and an associated callback function which will be automatically called when a response
+/// is received.
+class UnrealHttpRequestWithCallback {
+  const UnrealHttpRequestWithCallback(this.request, this.callback);
+
+  final UnrealHttpRequest request;
+  final Function(dynamic response) callback;
+}
+
+/// Holds the state of the app's connection to the engine and notifies when the connection changes.
+class EngineConnectionManager with WidgetsBindingObserver {
+  EngineConnectionManager(this.context) {
+    _loadLastConnectionIfNone();
+
+    WidgetsBinding.instance.addObserver(this);
+  }
+
+  /// The context this was built in.
+  final BuildContext context;
+
+  /// Name of the JSON field identifying the request ID of an HTTP request over WebSocket.
+  static const _httpRequestIdFieldName = 'RequestId';
+
+  /// The API version of the currently connected engine.
+  EpicStageAppAPIVersion? _apiVersion;
+
+  /// The data for the last active connection.
+  ConnectionData? _lastConnectionData;
+
+  /// The websocket channel we currently have open to the engine.
+  IOWebSocketChannel? _webSocketChannel;
+
+  /// Map from message type to callback functions for when the message type is received.
+  final Map<String, List<WebSocketMessageCallback>> _messageCallbacks = {};
+
+  /// Callbacks for whenever the connection state changes.
+  final List<ConnectionStateCallback> _connectionStateCallbacks = [];
+
+  /// Map from waiting request ID to the completer to complete when the corresponding response is received.
+  final Map<int, Completer<UnrealHttpResponse>> _httpResponseCompleters = {};
+
+  /// List of HTTP request IDs that must be processed before connection is complete.
+  final Set<int> _handshakeMessageIds = {};
+
+  /// List of messages that have been received, but are waiting to dispatch after connection setup is finished.
+  final List<dynamic> _initialPendingMessages = [];
+
+  /// The last request ID we used.
+  int _lastRequestId = 0;
+
+  /// Future that will complete pending a response to a message sent to check the engine's responsiveness, or null if
+  /// there's no pending message.
+  Future<void>? _respondingCheck;
+
+  /// The current state of the connection to the engine.
+  EngineConnectionState _internalConnectionState = EngineConnectionState.disconnected;
+
+  /// The current subscription to stream events from the connection.
+  StreamSubscription? _connectionSubscription;
+
+  /// Future that will return when the current connection attempt completes.
+  FutureOr<EngineConnectionResult>? _pendingConnectionAttempt;
+
+  /// HTTP client of the current WebSocket connection or connection attempt.
+  HttpClient? _webSocketHttpClient;
+
+  /// Whether the current connection attempt was cancelled.
+  bool _bIsPendingConnectionCancelled = false;
+
+  /// Whether this should attempt to reconnect when the app becomes active again.
+  bool _bShouldReconnectOnWake = false;
+
+  /// MD5 hash of the passphrase to be passed along with all messages for the current connection, or null if no
+  /// passphrase was given.
+  String? _passphraseHash;
+
+  /// Get the current state of the connection to the engine.
+  EngineConnectionState get connectionState => _internalConnectionState;
+
+  /// Set the current state of the connection to the engine.
+  set _connectionState(EngineConnectionState state) {
+    _internalConnectionState = state;
+
+    for (ConnectionStateCallback callback in _connectionStateCallbacks) {
+      callback(state);
+    }
+  }
+
+  /// Get the version of the engine we're currently connected to, or null if we're not connected.
+  EpicStageAppAPIVersion? get apiVersion => _apiVersion;
+
+  /// Dispose of any stored data.
+  void dispose() {
+    disconnect();
+    WidgetsBinding.instance.removeObserver(this);
+  }
+
+  @override
+  void didChangeAppLifecycleState(AppLifecycleState state) {
+    switch (state) {
+      case AppLifecycleState.detached:
+      case AppLifecycleState.paused:
+      case AppLifecycleState.inactive:
+        _onAppPaused();
+        break;
+
+      case AppLifecycleState.resumed:
+        _onAppResumed();
+        break;
+    }
+  }
+
+  /// Get the last engine connection we connected to, either in this session or a previous one.
+  Future<ConnectionData?> getLastConnectionData() {
+    if (_lastConnectionData != null) {
+      return Future.value(_lastConnectionData);
+    }
+
+    return _loadLastConnectionIfNone();
+  }
+
+  /// Try to connect with the given connection data. Returns a future which will provide the result of the attempt.
+  Future<EngineConnectionResult> connect(ConnectionData connectionData) async {
+    if (_pendingConnectionAttempt != null) {
+      if (!_bIsPendingConnectionCancelled) {
+        _log.warning('Tried to connect while an uncancelled connection attempt was in progress. New attempt ignored');
+        return EngineConnectionResult.alreadyOpen;
+      }
+
+      // Wait for cancelled attempt to finish so we don't try to open the same socket twice
+      await _pendingConnectionAttempt;
+    }
+
+    _bIsPendingConnectionCancelled = false;
+    Future<EngineConnectionResult> result;
+
+    try {
+      _pendingConnectionAttempt = _internalConnect(connectionData);
+      result = Future.value(await _pendingConnectionAttempt);
+    } catch (error) {
+      result = Future.value(EngineConnectionResult.genericFailure);
+    }
+
+    _pendingConnectionAttempt = null;
+
+    return result;
+  }
+
+  /// Try to reconnect to the last connection.
+  Future<EngineConnectionResult> reconnect() async {
+    if (_lastConnectionData != null) {
+      return connect(_lastConnectionData!);
+    } else {
+      return Future.error('Could not reconnect as there was no previous connection');
+    }
+  }
+
+  /// Disconnect from the engine.
+  void disconnect() {
+    if (_pendingConnectionAttempt != null) {
+      _bIsPendingConnectionCancelled = true;
+    }
+
+    _cancelHttpRequestsWithError(HttpResponseCode.clientClosedRequest);
+
+    if (_webSocketChannel != null) {
+      _log.info('Disconnecting from engine');
+
+      _webSocketChannel!.sink.close();
+      _connectionSubscription?.cancel();
+      _onWebSocketStreamClosed(
+        bWasExpected: true,
+        // If this was in the middle of attempting to connect, we were already on the connect screen and it can handle
+        // the navigator by itself
+        bShouldReturnToConnectScreen: _pendingConnectionAttempt == null,
+      );
+    }
+  }
+
+  /// Send a message via the current WebSocket connection in standard Unreal WebSocket format.
+  void sendMessage(String messageName, Object parameters) {
+    if (connectionState != EngineConnectionState.connected && _pendingConnectionAttempt == null) {
+      return;
+    }
+
+    sendRawMessage(createUnrealWebSocketMessage(messageName, parameters));
+  }
+
+  /// Send a batch of messages as a single message via the current WebSocket connection.
+  void sendBatchedMessage(List<dynamic> messages) {
+    if (connectionState != EngineConnectionState.connected && _pendingConnectionAttempt == null) {
+      return;
+    }
+
+    if (messages.length == 1) {
+      // No point in including the batch overhead.
+      sendRawMessage(messages[0]);
+      return;
+    }
+
+    sendMessage('batch', {'Requests': messages});
+  }
+
+  /// Send a raw message via the current WebSocket connection, encoded as JSON.
+  /// You may want to create the message using [createUnrealWebSocketMessage].
+  void sendRawMessage(dynamic message) {
+    if (connectionState != EngineConnectionState.connected && _pendingConnectionAttempt == null) {
+      return;
+    }
+
+    if (_passphraseHash != null) {
+      message['Passphrase'] = _passphraseHash;
+    }
+
+    _webSocketChannel?.sink.add(jsonEncode(message));
+  }
+
+  /// Send an HTTP request via the current WebSocket connection.
+  /// The response's body will contain decoded JSON data.
+  Future<UnrealHttpResponse> sendHttpRequest(UnrealHttpRequest request) {
+    return _internalSendHttpRequest(request);
+  }
+
+  /// Send an HTTP request via the current WebSocket connection and receive the response via a callback on the request.
+  void sendHttpRequestWithCallback(UnrealHttpRequestWithCallback request) {
+    sendHttpRequest(request.request).then(request.callback);
+  }
+
+  /// Send a batched HTTP request via the current WebSocket connection.
+  /// The response's body will contain a list of nullable [UnrealHttpResponse], where a response is null if there was a
+  /// problem parsing it.
+  Future<UnrealHttpResponse> sendBatchedHttpRequest(Iterable<UnrealHttpRequest> requests) {
+    final List<dynamic> requestParameters = [];
+
+    // Generate sub-requests
+    for (UnrealHttpRequest subRequest in requests) {
+      final parameters = _generateHttpRequestParameters(subRequest);
+      requestParameters.add(parameters);
+    }
+
+    // Set up the batch request and listen for its completion
+    final batchRequest = UnrealHttpRequest(url: '/remote/batch', verb: 'PUT', body: {'Requests': requestParameters});
+    final batchParameters = _generateHttpRequestParameters(batchRequest);
+    final message = {
+      'MessageName': 'http',
+      'Parameters': batchParameters,
+    };
+
+    final Completer<UnrealHttpResponse> completer = _makeHttpResponseCompleterFromParameters(batchParameters);
+    sendRawMessage(message);
+
+    return completer.future.then((batchResponse) {
+      if (batchResponse.code != HttpResponseCode.ok) {
+        return batchResponse;
+      }
+
+      final List? responses = batchResponse.body['Responses'];
+      if (responses == null) {
+        return Future.error('Got a valid batch response, but it contained no responses');
+      }
+
+      return UnrealHttpResponse(
+          requestId: batchResponse.requestId,
+          code: batchResponse.code,
+          body: responses.map((responseJson) => _decodeHttpResponseJsonData(responseJson)).toList());
+    });
+  }
+
+  /// Send a batched HTTP request via the current WebSocket connection and receive the response via callbacks on each
+  /// request.
+  void sendBatchedHttpRequestWithCallbacks(List<UnrealHttpRequestWithCallback> requests) {
+    final Iterable<UnrealHttpRequest> baseRequests =
+        requests.map((UnrealHttpRequestWithCallback requestWithCallback) => requestWithCallback.request);
+
+    sendBatchedHttpRequest(baseRequests).then((UnrealHttpResponse batchResponse) {
+      for (int responseIndex = 0; responseIndex < batchResponse.body.length; ++responseIndex) {
+        final UnrealHttpResponse? response = batchResponse.body[responseIndex];
+        requests[responseIndex].callback(response);
+      }
+    });
+  }
+
+  /// Check that the engine is responding to messages (i.e. not tied up in a queue or otherwise failing to acknowledge
+  /// them).
+  /// Returns a future that will complete when the engine responds to a message. Multiple calls to this function will
+  /// be combined into the most recent pending message.
+  Future<void> testEngineResponse() async {
+    if (_respondingCheck == null) {
+      // This message has a small response, so sending it isn't expensive
+      _respondingCheck = sendHttpRequest(UnrealHttpRequest(url: '/remote', verb: 'OPTIONS'));
+    }
+
+    await _respondingCheck;
+
+    _respondingCheck = null;
+  }
+
+  /// Register for callbacks when a type of WebSocket message is received.
+  void registerMessageListener(String messageType, WebSocketMessageCallback callback) {
+    if (_messageCallbacks.containsKey(messageType)) {
+      _messageCallbacks[messageType]!.add(callback);
+    } else {
+      _messageCallbacks[messageType] = [callback];
+    }
+  }
+
+  /// Unregister for callbacks when a type of WebSocket message is received.
+  void unregisterMessageListener(String messageType, WebSocketMessageCallback callback) {
+    if (_messageCallbacks.containsKey(messageType)) {
+      _messageCallbacks[messageType]!.remove(callback);
+    }
+  }
+
+  /// Register for callbacks when the connection state changes.
+  void registerConnectionStateListener(ConnectionStateCallback callback) {
+    _connectionStateCallbacks.add(callback);
+  }
+
+  /// Unregister for callbacks when the connection state changes.
+  void unregisterConnectionStateListener(ConnectionStateCallback callback) {
+    _connectionStateCallbacks.remove(callback);
+  }
+
+  /// Make and remember a new HTTP client to use for WebSocket connections.
+  HttpClient _makeHttpClient() {
+    assert(_webSocketHttpClient == null);
+
+    _webSocketHttpClient = HttpClient();
+    _webSocketHttpClient!.connectionTimeout = Duration(seconds: 3);
+
+    return _webSocketHttpClient!;
+  }
+
+  /// Internal function to connect to the engine.
+  /// You should almost always call [connect] instead of this so we know we have a pending connection attempt.
+  Future<EngineConnectionResult> _internalConnect(ConnectionData connectionData) async {
+    // Get the passphrase ready
+    final passphraseManager = Provider.of<EnginePassphraseManager>(context, listen: false);
+    final String? passphrase = await passphraseManager.getPassphrase(connectionData);
+
+    if (passphrase == null) {
+      _passphraseHash = null;
+    } else {
+      _passphraseHash = md5.convert(utf8.encode(passphrase)).toString();
+    }
+
+    // Close the existing connection
+    if (connectionState == EngineConnectionState.connected) {
+      _onWebSocketStreamClosed(bWasExpected: true, bShouldReturnToConnectScreen: false);
+    }
+
+    _log.info('Connecting to engine at ${connectionData.websocketAddress.address}:${connectionData.websocketPort}');
+
+    _lastConnectionData = connectionData;
+    _webSocketChannel = null;
+    _initialPendingMessages.clear();
+    _handshakeMessageIds.clear();
+
+    final String address = 'ws://${connectionData.websocketAddress.address}:${connectionData.websocketPort.toString()}';
+    WebSocket? webSocket;
+
+    // Try to connect
+    const int maxAttempts = 3;
+    bool bHasConnected = false;
+    String lastError = '';
+
+    for (int attempt = 0; attempt < maxAttempts && !_bIsPendingConnectionCancelled; ++attempt) {
+      _log.info('Connection attempt ${attempt + 1}...');
+      try {
+        webSocket = await WebSocket.connect(
+          address,
+          customClient: _makeHttpClient(),
+        );
+
+        bHasConnected = true;
+      } catch (error) {
+        _webSocketHttpClient?.close(force: true);
+        _webSocketHttpClient = null;
+
+        lastError = error.toString();
+        _log.info('Failed: $lastError');
+      }
+
+      if (bHasConnected) {
+        break;
+      }
+    }
+
+    if (_bIsPendingConnectionCancelled) {
+      if (bHasConnected) {
+        disconnect();
+      }
+      return EngineConnectionResult.cancelled;
+    }
+
+    if (!bHasConnected) {
+      return Future.error('Failed to connect to engine at $address: ' + lastError);
+    }
+
+    _log.info('Connected to $address');
+
+    // Set up the websocket channel
+    _webSocketChannel = IOWebSocketChannel(webSocket!);
+    _connectionSubscription = _webSocketChannel!.stream.listen(
+      _onWebSocketMessageReceived,
+      onError: _onWebSocketError,
+      onDone: _onWebSocketStreamClosed,
+    );
+
+    _saveLastConnection(connectionData);
+
+    // Check passphrase
+    final bool bIsPassphraseValid = await _checkPassphrase();
+    if (!bIsPassphraseValid) {
+      // The HTTP handler will have already disconnected us, so we can just return the result
+      return EngineConnectionResult.passphraseRejected;
+    }
+
+    // Retrieve API version
+    _log.info('Retrieving API version');
+    _apiVersion = await _retrieveAPIVersion();
+    _log.info('API version: $_apiVersion');
+
+    if (_bIsPendingConnectionCancelled) {
+      disconnect();
+      return EngineConnectionResult.genericFailure;
+    }
+
+    // Now that we're fully connected, handle any messages that were waiting
+    for (final dynamic message in _initialPendingMessages) {
+      _handleWebSocketMessage(message);
+    }
+
+    // This will notify external classes, so don't set it until everything is ready to go
+    _connectionState = EngineConnectionState.connected;
+
+    return EngineConnectionResult.success;
+  }
+
+  /// Called when a WebSocket error occurs.
+  void _onWebSocketError(dynamic error) {
+    _log.severe('WebSocket error', error);
+  }
+
+  /// Called when the WebSocket connection is closed.
+  void _onWebSocketStreamClosed({bool bWasExpected = false, bool bShouldReturnToConnectScreen = true}) async {
+    _webSocketChannel = null;
+    _connectionSubscription = null;
+    _connectionState = EngineConnectionState.disconnected;
+    _apiVersion = null;
+    _respondingCheck = null;
+    _passphraseHash = null;
+
+    // Force the HTTP client itself shut in case we were in the middle of connecting and want to immediately stop.
+    // The WebSocket itself isn't returned until the connection completes, so we have to kill the connection from here.
+    _webSocketHttpClient?.close(force: true);
+    _webSocketHttpClient = null;
+
+    if (bShouldReturnToConnectScreen) {
+      _goToNamedRoute(ConnectScreen.route);
+    }
+
+    if (!bWasExpected) {
+      InfoModalDialog.showInContext((context) => AppLocalizations.of(context)!.engineConnectionLostMessage);
+    }
+  }
+
+  /// Called when a WebSocket message is received.
+  void _onWebSocketMessageReceived(dynamic data) {
+    final String stringData = String.fromCharCodes(data);
+
+    final dynamic jsonMessage;
+    try {
+      jsonMessage = jsonDecode(stringData);
+    } catch (error) {
+      _log.warning('Failed to parse WebSocket message:\n$error\nMessage text:$stringData');
+      return;
+    }
+
+    // Passphrase failures are reported with no request ID and with a Verb of 401. In this case, all future messages
+    // will fail, so we can close out all completers with the same response and disconnect.
+    if (jsonMessage['Verb'] == '${HttpResponseCode.unauthorized}') {
+      _log.info('Passphrase rejected; disconnecting');
+
+      _cancelHttpRequestsWithError(HttpResponseCode.unauthorized);
+
+      // If we were mid connection attempt, the connection screen UI will handle showing the prompt if needed
+      final bool bShowPassphraseDialog = _pendingConnectionAttempt == null;
+
+      disconnect();
+
+      if (bShowPassphraseDialog) {
+        GenericModalDialogRoute.showDialog(
+          context: rootNavigatorKey.currentContext!,
+          builder: (context) => PassphraseModalDialog(
+            initialErrorMessage: AppLocalizations.of(context)!.connectScreenPassphraseDisconnectErrorMessage,
+          ),
+        );
+      }
+
+      return;
+    }
+
+    // If we're in the middle of connecting, we don't want to dispatch messages to outside systems.
+    // Unless the message is required for us to complete the connection handshake, queue it for later.
+    if (_pendingConnectionAttempt != null && !_handshakeMessageIds.contains(jsonMessage[_httpRequestIdFieldName])) {
+      _initialPendingMessages.add(jsonMessage);
+      return;
+    }
+
+    _handleWebSocketMessage(jsonMessage);
+  }
+
+  /// Handle a message received via WebSocket.
+  void _handleWebSocketMessage(dynamic jsonMessage) {
+    // Handle generic WebSocket messages
+    String? messageType = jsonMessage['Type'];
+    if (messageType != null) {
+      // Call any listeners of this message.
+      if (_messageCallbacks.containsKey(messageType)) {
+        for (WebSocketMessageCallback callback in _messageCallbacks[messageType]!) {
+          callback(jsonMessage);
+        }
+      }
+
+      return;
+    }
+
+    // Handle HTTP responses
+    final UnrealHttpResponse? decodedResponse = _decodeHttpResponseJsonData(jsonMessage);
+    if (decodedResponse != null && _httpResponseCompleters.containsKey(decodedResponse.requestId)) {
+      _httpResponseCompleters[decodedResponse.requestId]!.complete(decodedResponse);
+      _httpResponseCompleters.remove(decodedResponse.requestId);
+    }
+  }
+
+  /// Returns true if the current passphrase (if any) is accepted by the engine, and false otherwise.
+  /// Note that if the passphrase is rejected, we will automatically disconnect before this returns false.
+  Future<bool> _checkPassphrase() async {
+    final UnrealHttpResponse response = await _internalSendHttpRequest(
+      UnrealHttpRequest(
+        url: '/remote/passphrase',
+        verb: 'GET',
+      ),
+      bIsHandshakeMessage: true,
+    );
+
+    return response.code != HttpResponseCode.unauthorized;
+  }
+
+  /// Retrieve the API version from the engine, if possible.
+  Future<EpicStageAppAPIVersion> _retrieveAPIVersion() async {
+    const String stageAppLibraryPath = '/Script/EpicStageApp.Default__StageAppFunctionLibrary';
+    const EpicStageAppAPIVersion defaultVersion = EpicStageAppAPIVersion(1, 0, 0);
+
+    // Ask to describe the library first. If this fails, it won't dump an error in the editor log (unlike a failed
+    // remote call), and we'll know the engine is an old version without an ESA API version.
+    final UnrealHttpResponse describeResponse = await _internalSendHttpRequest(
+      UnrealHttpRequest(
+        url: '/remote/object/describe',
+        verb: 'PUT',
+        body: {
+          'objectPath': stageAppLibraryPath,
+        },
+      ),
+      bIsHandshakeMessage: true,
+    );
+
+    if (describeResponse.code != 200) {
+      // The library doesn't exist, so the engine can't report a version. Assume oldest API.
+      _log.info('Engine has no StageAppFunctionLibrary. Assuming default version.');
+      return defaultVersion;
+    }
+
+    final UnrealHttpResponse callResponse = await _internalSendHttpRequest(
+      UnrealHttpRequest(
+        url: '/remote/object/call',
+        verb: 'PUT',
+        body: {
+          'objectPath': stageAppLibraryPath,
+          'functionName': 'GetAPIVersion',
+        },
+      ),
+      bIsHandshakeMessage: true,
+    );
+
+    final String? returnValue = callResponse.body['ReturnValue'];
+    if (returnValue == null) {
+      _log.info('Failed to retrieve engine version. Assuming default version.');
+      return defaultVersion;
+    }
+
+    return EpicStageAppAPIVersion.fromString(returnValue);
+  }
+
+  /// Send an HTTP request via the current WebSocket connection.
+  /// The response's body will contain decoded JSON data.
+  /// If [bIsHandshakeMessage] is true, the request ID will be added to [_handshakeMessageIds].
+  Future<UnrealHttpResponse> _internalSendHttpRequest(UnrealHttpRequest request, {bool bIsHandshakeMessage = false}) {
+    final parameters = _generateHttpRequestParameters(request);
+    final completer = _makeHttpResponseCompleterFromParameters(parameters);
+
+    if (bIsHandshakeMessage) {
+      _handshakeMessageIds.add(_lastRequestId);
+    }
+
+    sendMessage('http', parameters);
+    return completer.future;
+  }
+
+  /// Decode JSON contained in an HTTP response body, or return null if the response was invalid.
+  UnrealHttpResponse? _decodeHttpResponseJsonData(dynamic jsonData) {
+    final int? requestId = jsonData[_httpRequestIdFieldName];
+    if (requestId == null) {
+      return null;
+    }
+
+    final int? responseCode = jsonData['ResponseCode'];
+    if (responseCode == null) {
+      return null;
+    }
+
+    // It's possible to leave this out, so we don't null check it
+    final dynamic responseBody = jsonData['ResponseBody'];
+
+    return UnrealHttpResponse(
+      requestId: requestId,
+      code: responseCode,
+      body: responseBody,
+    );
+  }
+
+  /// Generate an object containing request parameters for the given HTTP request.
+  dynamic _generateHttpRequestParameters(UnrealHttpRequest request) {
+    ++_lastRequestId;
+    return {
+      _httpRequestIdFieldName: _lastRequestId,
+      'Url': request.url,
+      'Verb': request.verb,
+      if (request.body != null) 'Body': request.body,
+    };
+  }
+
+  /// Generate a completer that will complete its future when a response is received for the request with the given
+  /// parameters.
+  Completer<UnrealHttpResponse> _makeHttpResponseCompleterFromParameters(dynamic parameters) {
+    final int? requestId = parameters[_httpRequestIdFieldName];
+
+    if (requestId == null) {
+      throw Exception('Expected a request ID to listen for');
+    }
+
+    final completer = Completer<UnrealHttpResponse>();
+    _httpResponseCompleters[requestId] = completer;
+
+    return completer;
+  }
+
+  /// Get the last connection stored in persisted  preferences and restore it into [_lastConnectionData] if we don't
+  /// already have a connection saved.
+  Future<ConnectionData?> _loadLastConnectionIfNone() async {
+    final sharedPrefs = await SharedPreferences.getInstance();
+
+    if (_lastConnectionData != null) {
+      // We already connected somewhere newer
+      return Future.value(_lastConnectionData);
+    }
+
+    // Get the last connection's data
+    final String? name = sharedPrefs.getString('lastConnection.name');
+    final String? address = sharedPrefs.getString('lastConnection.websocketAddress');
+    final int? port = sharedPrefs.getInt('lastConnection.websocketPort');
+
+    if (name != null && address != null && port != null) {
+      ConnectionData? loadedConnectionData;
+      try {
+        loadedConnectionData = ConnectionData(
+          uuid: UuidValue(const Uuid().v1()),
+          name: name,
+          websocketAddress: InternetAddress(address),
+          websocketPort: port,
+        );
+      } catch (error) {
+        _log.warning('Failed to restore connection data from preferences', error);
+      }
+
+      if (loadedConnectionData != null) {
+        _lastConnectionData = loadedConnectionData;
+      }
+    }
+
+    return Future.value(_lastConnectionData);
+  }
+
+  /// Save the last connection to persisted preferences.
+  void _saveLastConnection(ConnectionData lastConnection) async {
+    final sharedPrefs = await SharedPreferences.getInstance();
+
+    sharedPrefs.setString('lastConnection.name', lastConnection.name);
+    sharedPrefs.setString('lastConnection.websocketAddress', lastConnection.websocketAddress.address);
+    sharedPrefs.setInt('lastConnection.websocketPort', lastConnection.websocketPort);
+  }
+
+  /// Complete all pending HTTP-over-WebSocket requests with the same error response.
+  void _cancelHttpRequestsWithError(int errorCode) {
+    for (final MapEntry<int, Completer<UnrealHttpResponse>> httpCompleterEntry in _httpResponseCompleters.entries) {
+      final response = UnrealHttpResponse(
+        requestId: httpCompleterEntry.key,
+        code: errorCode,
+        body: {},
+      );
+      httpCompleterEntry.value.complete(response);
+    }
+    _httpResponseCompleters.clear();
+  }
+
+  /// Called when the app is paused.
+  void _onAppPaused() {
+    if (connectionState == EngineConnectionState.disconnected && _pendingConnectionAttempt == null) {
+      return;
+    }
+
+    _bShouldReconnectOnWake = true;
+  }
+
+  /// Called when the app resumes after being paused.
+  void _onAppResumed() async {
+    if (!_bShouldReconnectOnWake) {
+      return;
+    }
+
+    _bShouldReconnectOnWake = false;
+    if (connectionState == EngineConnectionState.connected || _pendingConnectionAttempt != null) {
+      // We're still connected/connecting
+      return;
+    }
+
+    // Show a temporary screen while we try to reconnect
+    _goToNamedRoute(ReconnectScreen.route);
+
+    try {
+      await reconnect();
+    } catch (e) {
+      // Failed to connect, which will be handled below
+    }
+
+    if (connectionState == EngineConnectionState.connected) {
+      _goToNamedRoute(StageAppMainScreen.route);
+    } else {
+      disconnect();
+      _goToNamedRoute(ConnectScreen.route);
+    }
+  }
+
+  /// Pop all screens and push a named route.
+  void _goToNamedRoute(String name) {
+    rootNavigatorKey.currentState?.pushNamedAndRemoveUntil(
+      name,
+      (route) => false,
+    );
+  }
+}
diff --git a/Engine/Extras/VirtualProduction/EpicStageApp/lib/models/engine_passphrase_manager.dart b/Engine/Extras/VirtualProduction/EpicStageApp/lib/models/engine_passphrase_manager.dart
new file mode 100644
--- /dev/null
+++ b/Engine/Extras/VirtualProduction/EpicStageApp/lib/models/engine_passphrase_manager.dart
@@ -0,0 +1,31 @@
+// Copyright Epic Games, Inc. All Rights Reserved.
+
+import 'package:flutter_secure_storage/flutter_secure_storage.dart';
+
+import '../utilities/net_utilities.dart';
+
+/// Handles storing and retrieving passphrases to use for each engine instance.
+class EnginePassphraseManager {
+  final FlutterSecureStorage _secureStorage = FlutterSecureStorage();
+
+  /// Get the passphrase associated with the given [connectionData], or null if one isn't saved.
+  Future<String?> getPassphrase(ConnectionData connectionData) {
+    return _secureStorage.read(key: _makeKeyForConnection(connectionData));
+  }
+
+  /// Set the [passphrase] associated with the given [connectionData], or delete it if no passphrase is given.
+  void setPassphrase(ConnectionData connectionData, String? passphrase) {
+    final String key = _makeKeyForConnection(connectionData);
+
+    if (passphrase == null) {
+      _secureStorage.delete(key: key);
+    } else {
+      _secureStorage.write(key: key, value: passphrase);
+    }
+  }
+
+  /// Given [connectionData], create a unique key to refer to its passphrase in the secure storage database.
+  String _makeKeyForConnection(ConnectionData connectionData) {
+    return 'enginePassphrase.${connectionData.websocketAddress.address}';
+  }
+}
diff --git a/Engine/Extras/VirtualProduction/EpicStageApp/lib/models/navigator_keys.dart b/Engine/Extras/VirtualProduction/EpicStageApp/lib/models/navigator_keys.dart
new file mode 100644
--- /dev/null
+++ b/Engine/Extras/VirtualProduction/EpicStageApp/lib/models/navigator_keys.dart
@@ -0,0 +1,6 @@
+// Copyright Epic Games, Inc. All Rights Reserved.
+
+import 'package:flutter/cupertino.dart';
+
+/// Key used to access the root navigator for the app without having to retrieve it from the BuildContext.
+final GlobalKey<NavigatorState> rootNavigatorKey = GlobalKey<NavigatorState>();
