Branch snapshot for CL 27138787

diff --git a/Engine/Extras/VirtualProduction/EpicStageApp/lib/models/preview_render_manager.dart b/Engine/Extras/VirtualProduction/EpicStageApp/lib/models/preview_render_manager.dart
new file mode 100644
--- /dev/null
+++ b/Engine/Extras/VirtualProduction/EpicStageApp/lib/models/preview_render_manager.dart
@@ -0,0 +1,657 @@
+// Copyright Epic Games, Inc. All Rights Reserved.
+
+import 'dart:async';
+import 'dart:convert';
+import 'dart:math' as math;
+import 'dart:typed_data';
+
+import 'package:flutter/material.dart';
+import 'package:logging/logging.dart';
+import 'package:provider/provider.dart';
+import 'package:vector_math/vector_math_64.dart' as vec;
+
+import '../utilities/constants.dart';
+import '../utilities/net_utilities.dart';
+import '../utilities/unreal_utilities.dart';
+import 'engine_connection.dart';
+import 'settings/selected_actor_settings.dart';
+import 'settings/stage_map_settings.dart';
+import 'unreal_actor_manager.dart';
+import 'unreal_types.dart';
+
+final _log = Logger('PreviewRender');
+
+/// Maintains information about and provides event hooks for actors in the engine that the app cares about.
+class PreviewRenderManager with WidgetsBindingObserver {
+  /// Quality of the JPEG to request (in range 50-100).
+  static const int _jpegQuality = 50;
+
+  /// How long to wait after receiving a preview render before requesting a new one, as a factor of the response time
+  /// for the last render.
+  static const double _previewRequestDelayFactor = 0.5;
+
+  /// The maximum number of seconds to wait before requesting a new preview render, regardless of how long the previous
+  /// request took.
+  static const double _maxPreviewRequestDelay = 1.0;
+
+  /// The maximum time to wait for the engine to respond with a render before sending a new request, in case the
+  /// previous request was lost.
+  static const Duration _previewRequestTimeout = Duration(seconds: 3);
+
+  /// The maximum consecutive number of timeouts before we give up and create a new preview renderer.
+  static const int _maxPreviewTimeoutsBeforeGiveUp = 3;
+
+  /// Rate at which we send light card/camera config updates back to the engine.
+  static const Duration _engineUpdateRate = Duration(milliseconds: 100);
+
+  /// The current consumers for preview renders, if any.
+  Set<PreviewRenderConsumer> _consumers = {};
+
+  /// The last root actor path that we first sent to the engine, stored so we can check if we need to update it once the
+  /// engine replies with a renderer ID.
+  String? _initialRootPath;
+
+  /// The ID of our preview renderer.
+  int? _previewRendererId;
+
+  /// The last time we sent a request to render a new preview of the scene.
+  DateTime _lastRenderRequestTime = DateTime.now();
+
+  /// Timer used to regularly send changes back to the engine while connected.
+  Timer? _engineUpdateTimer;
+
+  /// Whether we've sent a request to create the preview renderer.
+  bool _bHasCreatedRenderer = false;
+
+  /// Whether an actor drag operation is in progress.
+  bool _bIsDraggingActors = false;
+
+  /// The size of preview to request from the engine. The engine may provide a lower resolution with the same aspect
+  /// ratio.
+  Size _desiredResolution = const Size(1920, 1080);
+
+  /// Helper currently attempting to initialize a preview renderer.
+  MessageTimeoutRetryHelper? _initializeRendererHelper;
+
+  /// Helper currently attempting to request a preview render.
+  MessageTimeoutRetryHelper? _requestRenderHelper;
+
+  /// The last sequence number we sent the engine for a drag update.
+  int _dragSequenceNumber = 0;
+
+  /// Stream subscriptions to user settings.
+  final List<StreamSubscription> _settingSubscriptions = [];
+
+  /// The connection manager we'll use to communicate with the engine.
+  final EngineConnectionManager _connectionManager;
+
+  /// Actor manager used to watch for root actor changes.
+  final UnrealActorManager _actorManager;
+
+  /// Settings for which actors are selected.
+  final SelectedActorSettings _selectedActorSettings;
+
+  /// The stage map settings, which include camera controls.
+  final StageMapSettings _stageMapSettings;
+
+  /// Modified renderer settings to send to the engine on the next tick. This doesn't have to include all settings; just
+  /// the ones that we want to change.
+  final Map<String, dynamic> _tickRendererSettingsChanges = {};
+
+  /// The ID of the preview renderer, or null if it hasn't been created.
+  int? get rendererId => _previewRendererId;
+
+  /// The focal point stored in the stage map settings and adjusted for the current renderer's aspect ratio.
+  Offset get rendererFocalPoint {
+    final Offset focalPoint = _stageMapSettings.focalPoint.getValue();
+
+    return Offset(
+      0.5 + (focalPoint.dx - 0.5) / _desiredResolution.aspectRatio,
+      focalPoint.dy,
+    );
+  }
+
+  /// Whether an actor drag operation is in progress.
+  bool get bIsDraggingActors => _bIsDraggingActors;
+
+  /// Construct the manager and register it with the connection manager.
+  PreviewRenderManager(BuildContext context)
+      : _connectionManager = Provider.of<EngineConnectionManager>(context, listen: false),
+        _stageMapSettings = Provider.of<StageMapSettings>(context, listen: false),
+        _selectedActorSettings = Provider.of<SelectedActorSettings>(context, listen: false),
+        _actorManager = Provider.of<UnrealActorManager>(context, listen: false) {
+    _startListeningToEngineConnection();
+
+    _onConnectionStateChanged(_connectionManager.connectionState);
+
+    _actorManager.watchClassName(nDisplayRootActorClassName, _onRootActorUpdate);
+
+    _settingSubscriptions.addAll([
+      _selectedActorSettings.displayClusterRootPath.listen(_onRootActorPathChanged),
+      _stageMapSettings.projectionMode.listen(_onProjectionModeChanged),
+      _stageMapSettings.cameraAngle.listen(_onCameraAngleChanged),
+    ]);
+
+    WidgetsBinding.instance.addObserver(this);
+  }
+
+  /// Register a consumer with the manager, indicating that it should receive render updates.
+  void addConsumer(PreviewRenderConsumer newConsumer) {
+    _consumers.add(newConsumer);
+
+    // Request a preview render if we weren't already waiting for one
+    if (_requestRenderHelper?.bIsActive != true) {
+      _log.info('Activated preview renderer and started new preview render request');
+      _requestPreviewRender();
+    }
+  }
+
+  /// Unregister a consumer with the manager.
+  void removeConsumer(PreviewRenderConsumer consumer) {
+    _consumers.remove(consumer);
+  }
+
+  /// Change the desired preview resolution.
+  /// Returns whether the new resolution was valid. If not, the resolution won't be changed.
+  bool setDesiredResolution(Size resolution) {
+    if (resolution.width <= 0 || resolution.height <= 0) {
+      return false;
+    }
+
+    if (_desiredResolution != resolution) {
+      _desiredResolution = resolution;
+
+      _tickRendererSettingsChanges['Resolution'] = {
+        'X': _desiredResolution.width,
+        'Y': _desiredResolution.height,
+      };
+    }
+
+    return true;
+  }
+
+  /// Send a message to the engine beginning a drag operation on the given [actors] (paths) and using the [primaryActor]
+  /// as the fulcrum point.
+  /// A drag must not already be in progress, which you can check with [bIsDraggingActors].
+  void beginActorDrag({required List<String> actors, String? primaryActor}) {
+    if (rendererId == null) {
+      _log.warning('Tried to start a drag operation, but there was no renderer ID');
+      return;
+    }
+
+    if (bIsDraggingActors) {
+      _log.warning('Tried to start a drag operation, but one was already in progress');
+      return;
+    }
+
+    _bIsDraggingActors = true;
+
+    ++_dragSequenceNumber;
+    _connectionManager.sendMessage('ndisplay.preview.actor.drag.begin', {
+      'RendererId': rendererId!,
+      'Actors': actors,
+      'PrimaryActor': primaryActor ?? actors.last,
+      'SequenceNumber': _dragSequenceNumber,
+    });
+  }
+
+  /// Send a message to the engine updating the normalized [position] of the current actor drag operation within the
+  /// preview. A drag must already be in progress, which you can check with [bIsDraggingActors].
+  void updateActorDrag(Offset position) {
+    if (rendererId == null) {
+      _log.warning('Tried to update a drag operation, but there was no renderer ID');
+      return;
+    }
+
+    if (!bIsDraggingActors) {
+      _log.warning('Tried to update a drag operation, but there was not one in progress');
+      return;
+    }
+
+    // Send the latest actor drag position
+    ++_dragSequenceNumber;
+    _connectionManager.sendMessage('ndisplay.preview.actor.drag.move', {
+      'RendererId': rendererId!,
+      'DragPosition': offsetToJson(position),
+      'SequenceNumber': _dragSequenceNumber,
+    });
+  }
+
+  /// Send a message to the engine ending the current actor drag operation at the given normalized [position] within the
+  /// preview.
+  /// A drag must already be in progress, which you can check with [bIsDraggingActors].
+  void endActorDrag(Offset position) {
+    if (rendererId == null) {
+      _log.warning('Tried to end a drag operation, but there was no renderer ID');
+      return;
+    }
+
+    if (!bIsDraggingActors) {
+      _log.warning('Tried to end a drag operation, but there was not one in progress');
+      return;
+    }
+
+    _bIsDraggingActors = false;
+
+    // Send the latest actor drag position
+    ++_dragSequenceNumber;
+    _connectionManager.sendMessage('ndisplay.preview.actor.drag.end', {
+      'RendererId': rendererId!,
+      'DragPosition': offsetToJson(position),
+      'SequenceNumber': _dragSequenceNumber,
+    });
+  }
+
+  @override
+  void didChangeAppLifecycleState(AppLifecycleState state) {
+    super.didChangeAppLifecycleState(state);
+
+    switch (state) {
+      case AppLifecycleState.resumed:
+        if (_connectionManager.connectionState == EngineConnectionState.connected && !_bHasCreatedRenderer) {
+          _initializePreviewRenderer();
+        }
+        break;
+
+      case AppLifecycleState.paused:
+      case AppLifecycleState.detached:
+        _destroyPreviewRenderer();
+        break;
+
+      default:
+        break;
+    }
+  }
+
+  /// Call this before the manager will be destroyed.
+  void dispose() {
+    _onDisconnected();
+
+    _stopListeningToEngineConnection();
+
+    for (final StreamSubscription subscription in _settingSubscriptions) {
+      subscription.cancel();
+    }
+  }
+
+  /// Start listening for messages from the engine connection manager.
+  void _startListeningToEngineConnection() {
+    _connectionManager
+      ..registerMessageListener('PreviewRendererCreated', _onPreviewRendererCreated)
+      ..registerMessageListener('PreviewRenderCompleted', _onPreviewRenderCompleted)
+      ..registerConnectionStateListener(_onConnectionStateChanged);
+  }
+
+  /// Stop listening for messages from the engine connection manager.
+  void _stopListeningToEngineConnection() {
+    _connectionManager
+      ..unregisterMessageListener('PreviewRendererCreated', _onPreviewRendererCreated)
+      ..unregisterMessageListener('PreviewRenderCompleted', _onPreviewRenderCompleted)
+      ..unregisterConnectionStateListener(_onConnectionStateChanged);
+  }
+
+  /// Request that the engine destroy the preview renderer.
+  void _destroyPreviewRenderer() {
+    if (_previewRendererId != null && _connectionManager.connectionState == EngineConnectionState.connected) {
+      _connectionManager.sendMessage('ndisplay.preview.renderer.destroy', {'RendererId': _previewRendererId!});
+    }
+
+    _requestRenderHelper?.giveUp();
+
+    _previewRendererId = null;
+    _bHasCreatedRenderer = false;
+    _bIsDraggingActors = false;
+  }
+
+  /// Called when the connection manager's connection state changes.
+  void _onConnectionStateChanged(EngineConnectionState connectionState) {
+    switch (connectionState) {
+      case EngineConnectionState.connected:
+        _onConnected();
+        break;
+
+      default:
+        _onDisconnected();
+        break;
+    }
+  }
+
+  /// Called when we connect to the engine.
+  void _onConnected() {
+    _engineUpdateTimer ??= Timer.periodic(_engineUpdateRate, _onEngineUpdateTick);
+
+    if (!_bHasCreatedRenderer) {
+      _initializePreviewRenderer();
+    }
+  }
+
+  /// Called when we disconnect from the engine.
+  void _onDisconnected() {
+    _destroyPreviewRenderer();
+
+    _engineUpdateTimer?.cancel();
+    _engineUpdateTimer = null;
+  }
+
+  /// Send a request to initialize the preview renderer in the engine.
+  void _initializePreviewRenderer() {
+    final rootPath = _selectedActorSettings.displayClusterRootPath.getValue();
+    if (rootPath.isEmpty) {
+      return;
+    }
+
+    _log.info('Requesting preview renderer');
+
+    _initializeRendererHelper = MessageTimeoutRetryHelper(
+      logger: _log,
+      logDescription: 'initialize preview renderer',
+      connection: _connectionManager,
+      timeoutDuration: _previewRequestTimeout,
+      sendMessage: () {
+        final rootPath = _selectedActorSettings.displayClusterRootPath.getValue();
+        if (rootPath.isEmpty) {
+          _log.info('No root path while attempting to request preview renderer; giving up for now');
+          _initializeRendererHelper?.giveUp();
+          _destroyPreviewRenderer();
+          return;
+        }
+
+        _connectionManager.sendMessage('ndisplay.preview.renderer.create', {
+          'RootActorPath': rootPath,
+          'Settings': {
+            'JpegQuality': _jpegQuality,
+            'Resolution': {
+              'X': _desiredResolution.width,
+              'Y': _desiredResolution.height,
+            },
+            'Rotation': _getCameraRotationSettings(),
+            'IncludeActorPositions': true,
+            'ProjectionType': _getProjectionModeTypeName(_stageMapSettings.projectionMode.getValue()),
+            'FOV': _getFOV(),
+          },
+        });
+      },
+    );
+    _initializeRendererHelper!.start();
+
+    _initialRootPath = rootPath;
+    _bHasCreatedRenderer = true;
+  }
+
+  /// Called when the engine has successfully created a preview renderer.
+  void _onPreviewRendererCreated(dynamic message) {
+    final int? rendererId = message['RendererId'];
+    if (rendererId == null) {
+      return;
+    }
+
+    _initializeRendererHelper?.succeed();
+
+    _log.info('Preview renderer #$rendererId ready');
+
+    _previewRendererId = rendererId;
+    _dragSequenceNumber = 0;
+
+    final String rootPath = _selectedActorSettings.displayClusterRootPath.getValue();
+    if (_initialRootPath != rootPath) {
+      // Root path has changed since we requested the renderer. Update it.
+      _onRootActorPathChanged(rootPath);
+    }
+
+    // Request a render to kick off our render loop.
+    _requestPreviewRender();
+  }
+
+  /// Request a preview render from the engine (if there's anybody to consume it).
+  void _requestPreviewRender() {
+    if (_previewRendererId == null || _consumers.isEmpty || _requestRenderHelper?.bIsActive == true) {
+      return;
+    }
+
+    _requestRenderHelper = MessageTimeoutRetryHelper(
+      logger: _log,
+      logDescription: 'new preview render',
+      connection: _connectionManager,
+      timeoutDuration: _previewRequestTimeout,
+      sendMessage: () {
+        _lastRenderRequestTime = DateTime.now();
+        _connectionManager.sendMessage('ndisplay.preview.render', {'RendererId': _previewRendererId!});
+      },
+      onTimeout: (int timeoutCount) {
+        if (timeoutCount > _maxPreviewTimeoutsBeforeGiveUp) {
+          _log.info('Too many failed preview render attempts. Creating a new preview renderer.');
+          _requestRenderHelper?.giveUp();
+          _destroyPreviewRenderer();
+          _initializePreviewRenderer();
+        }
+      },
+    );
+
+    _requestRenderHelper?.start();
+  }
+
+  /// Called when the engine completes a preview render of the scene.
+  void _onPreviewRenderCompleted(dynamic message) {
+    if (_previewRendererId == null) {
+      // We're ignoring new renders at the moment, probably because we were disconnected
+      return;
+    }
+
+    final int? receivedRendererId = message['RendererId'];
+    if (receivedRendererId != _previewRendererId) {
+      _log.warning('Received preview for wrong renderer (ours: $_previewRendererId; received: $receivedRendererId)');
+      return;
+    }
+
+    _requestRenderHelper?.succeed();
+    _requestRenderHelper = null;
+
+    final DateTime receivedTime = DateTime.now();
+
+    if (_consumers.isNotEmpty) {
+      final String? base64Data = message['ImageBase64'];
+      final Uint8List imageData = base64Data != null ? base64Decode(base64Data) : Uint8List(0);
+
+      final num sequenceNumber = message['SequenceNumber'] ?? -1;
+      final bool bIsRealTimeDisabled = message['IsRealTimeDisabled'] == true;
+
+      final dynamic resolutionData = message['Resolution'];
+      Size resolution;
+      if (resolutionData != null) {
+        resolution = Size(
+          jsonNumberToDouble(resolutionData['X'])!,
+          jsonNumberToDouble(resolutionData['Y'])!,
+        );
+      } else {
+        // No resolution was sent, so assume it's the same as our upcoming request. Note that this can cause pins to
+        // jump when the containing widget is resized.
+        resolution = _desiredResolution;
+      }
+
+      // Gather actor position data
+      final List<PreviewActorData> actorPositions = [];
+      final dynamic actorPositionsData = message['ActorPositions'];
+      for (dynamic positionPair in actorPositionsData) {
+        final String? actorPath = positionPair['Path'];
+        if (actorPath == null) {
+          continue;
+        }
+
+        final dynamic actorPositionData = positionPair['Position'];
+        if (actorPositionData == null) {
+          continue;
+        }
+
+        final double? xPosition = jsonNumberToDouble(actorPositionData['X']);
+        final double? yPosition = jsonNumberToDouble(actorPositionData['Y']);
+        if (xPosition == null || yPosition == null) {
+          continue;
+        }
+
+        actorPositions.add(PreviewActorData(
+          path: actorPath,
+          position: Offset(xPosition, yPosition),
+        ));
+      }
+
+      for (final PreviewRenderConsumer consumer in _consumers) {
+        consumer.onPreviewRenderCompleted(
+          bIsForLatestRequest: sequenceNumber.toInt() >= _dragSequenceNumber,
+          resolution: resolution,
+          imageData: imageData,
+          actorPositions: actorPositions,
+          bIsRealTimeDisabled: bIsRealTimeDisabled,
+        );
+      }
+    }
+
+    // Wait for a bit before we request another preview so we don't eat up too much bandwidth
+    final Duration timeSinceRequest = receivedTime.difference(_lastRenderRequestTime);
+    final double delayMilliseconds = math.min(timeSinceRequest.inMilliseconds * _previewRequestDelayFactor,
+        _maxPreviewRequestDelay * Duration.millisecondsPerSecond);
+
+    Future.delayed(Duration(milliseconds: delayMilliseconds.toInt()), _requestPreviewRender);
+  }
+
+  /// Called each tick, when we're ready to send updates back to the engine.
+  void _onEngineUpdateTick(Timer timer) {
+    if (_previewRendererId == null || _tickRendererSettingsChanges.isEmpty) {
+      return;
+    }
+
+    // Send the latest renderer config changes
+    _connectionManager.sendMessage('ndisplay.preview.renderer.configure', {
+      'RendererId': _previewRendererId,
+      'Settings': _tickRendererSettingsChanges,
+    });
+
+    _tickRendererSettingsChanges.clear();
+  }
+
+  /// Get the name to send to the engine to request a projection mode
+  String _getProjectionModeTypeName(ProjectionMode mode) {
+    switch (mode) {
+      case ProjectionMode.azimuthal:
+        return 'Azimuthal';
+
+      case ProjectionMode.orthographic:
+        return 'Orthographic';
+
+      case ProjectionMode.perspective:
+        return 'Perspective';
+
+      case ProjectionMode.uv:
+        return 'UV';
+    }
+  }
+
+  /// Update the preview renderer camera's angle.
+  void _onCameraAngleChanged(vec.Vector2 cameraAngle) {
+    _tickRendererSettingsChanges['Rotation'] = _getCameraRotationSettings();
+  }
+
+  /// Update the projection mode and affected settings.
+  void _onProjectionModeChanged(ProjectionMode projectionMode) {
+    _tickRendererSettingsChanges['ProjectionType'] = _getProjectionModeTypeName(projectionMode);
+    _tickRendererSettingsChanges['FOV'] = _getFOV();
+    _onCameraAngleChanged(_stageMapSettings.cameraAngle.getValue());
+  }
+
+  /// Called when the root actors in the engine have changed.
+  void _onRootActorUpdate(ActorUpdateDetails details) {
+    if (_previewRendererId == null || _connectionManager.connectionState != EngineConnectionState.connected) {
+      return;
+    }
+
+    final String rootPath = _selectedActorSettings.displayClusterRootPath.getValue();
+    if (details.addedActors.isEmpty || rootPath.isEmpty) {
+      return;
+    }
+
+    if (!_bHasCreatedRenderer) {
+      // May have been waiting for valid settings
+      _initializePreviewRenderer();
+    }
+
+    final UnrealObject? currentRoot = _actorManager.getActorAtPath(rootPath);
+    if (currentRoot == null) {
+      return;
+    }
+
+    if (!details.addedActors.contains(currentRoot)) {
+      return;
+    }
+
+    // If the root actor has been re-added to the scene, our render will be cancelled, so we need to request a new one.
+    _requestPreviewRender();
+  }
+
+  /// Called when the path to the root nDisplay actor changes.
+  void _onRootActorPathChanged(String rootActorPath) {
+    if (_previewRendererId != null) {
+      _connectionManager.sendMessage('ndisplay.preview.renderer.setroot', {
+        'RendererId': _previewRendererId!,
+        'RootActorPath': rootActorPath,
+      });
+
+      // Request a new render with the new root, since changing roots may mean we never receive the last render
+      _requestPreviewRender();
+    }
+  }
+
+  /// Get the field of view based on the current projection mode.
+  double _getFOV() {
+    switch (_stageMapSettings.projectionMode.getValue()) {
+      case ProjectionMode.uv:
+        return 45.0;
+
+      case ProjectionMode.azimuthal:
+        return 130.0;
+
+      case ProjectionMode.orthographic:
+      case ProjectionMode.perspective:
+        return 90.0;
+    }
+  }
+
+  /// Get the JSON data representing the camera's rotation settings based on the projection mode and user settings.
+  dynamic _getCameraRotationSettings() {
+    final bool bIsUvMode = _stageMapSettings.projectionMode.getValue() == ProjectionMode.uv;
+    final vec.Vector2 cameraAngle = _stageMapSettings.cameraAngle.getValue();
+
+    return {
+      'Yaw': bIsUvMode ? 0 : cameraAngle.x,
+      'Pitch': bIsUvMode ? 0 : cameraAngle.y,
+      'Roll': 0,
+    };
+  }
+}
+
+/// Data about an actor's position as seen in a preview render.
+class PreviewActorData {
+  const PreviewActorData({required this.path, required this.position});
+
+  /// The path of the actor object.
+  final String path;
+
+  /// The actor's position within the preview render, normalized to the resolution of the preview.
+  final Offset position;
+}
+
+/// Mixin for classes that want to receive renders from [PreviewRenderManager].
+abstract class PreviewRenderConsumer {
+  /// Called when a preview render from the engine completes.
+  /// [bIsForLatestRequest] is true if this contains the latest available preview and actor position data, or false if
+  /// the engine is expected to send more recent data after this.
+  /// [resolution] is the resolution of the preview image.
+  /// [imageData] is the data for the preview image.
+  /// [actorPositions] is a list of position data for all actors visible in the preview.
+  /// [bIsRealTimeDisabled] is whether real-time rendering was disabled in the engine when the preview was rendered.
+  void onPreviewRenderCompleted({
+    required bool bIsForLatestRequest,
+    required Size resolution,
+    required Uint8List imageData,
+    required List<PreviewActorData> actorPositions,
+    required bool bIsRealTimeDisabled,
+  });
+}
diff --git a/Engine/Extras/VirtualProduction/EpicStageApp/lib/models/property_modify_operations.dart b/Engine/Extras/VirtualProduction/EpicStageApp/lib/models/property_modify_operations.dart
new file mode 100644
--- /dev/null
+++ b/Engine/Extras/VirtualProduction/EpicStageApp/lib/models/property_modify_operations.dart
@@ -0,0 +1,428 @@
+// Copyright Epic Games, Inc. All Rights Reserved.
+
+import 'dart:ui';
+
+import '../utilities/math_utilities.dart';
+import '../utilities/unreal_utilities.dart';
+import 'color.dart';
+import 'unreal_types.dart';
+
+/// How a property's value should be treated when it pushes past the min/max for the property.
+enum PropertyMinMaxBehaviour {
+  /// Ignore the engine min/max values entirely.
+  ignore,
+
+  /// Clamp the value to the min/max value.
+  clamp,
+
+  /// Use a modulo operation to "loop" back to the other side, i.e. pushing past max will return to min and vice versa.
+  loop,
+}
+
+/// Data about a modification event for an Unreal property. This data is generated when we're ready to send a message
+/// back to the engine about the change.
+class PropertyModifyEvent {
+  const PropertyModifyEvent({
+    required this.property,
+    required this.currentValue,
+    required this.deltaValue,
+    required this.presetName,
+    required this.presetPropertyLabel,
+    required this.typeName,
+    required this.sequenceNumber,
+    required this.transactionId,
+  });
+
+  /// The property being modified.
+  final UnrealProperty property;
+
+  /// The current value of the property.
+  final dynamic currentValue;
+
+  /// The delta value for this change event.
+  final dynamic deltaValue;
+
+  /// The name of the preset that the property is exposed on.
+  final String presetName;
+
+  /// The label of the property as exposed on the preset.
+  final String presetPropertyLabel;
+
+  /// The type name of the property being modified.
+  final String typeName;
+
+  /// The sequence number of this change.
+  final int sequenceNumber;
+
+  /// The ID of the transaction with which to associate the changes.
+  final int transactionId;
+}
+
+/// Defines logic for a type of property modification, including applying the change locally, clamping it to reasonable
+/// bounds, and sending the result to the engine.
+abstract class PropertyModifyOperation {
+  const PropertyModifyOperation();
+
+  /// If true, this operation will always be considered significant even if the clamp returns null.
+  bool get bIsAlwaysSignificant => false;
+
+  /// The name of the WebSocket message to send for this operation.
+  String get messageName;
+
+  /// Returns the result of applying a change of [deltaValue] to a base [value] using this operation.
+  dynamic apply(dynamic value, dynamic deltaValue);
+
+  /// Implements clamping behaviour for this operation. Given a [value] that the [deltaValue] will be applied to, return
+  /// a new delta value which won't exceed the constraints by clamping the result value between min and max.
+  /// If the delta is effectively a no-op as a result of the clamp, return null.
+  /// [hardMin] and [hardMax] contain the min/max values allowed for this property, though they may be null meaning
+  /// the property has no limit in the corresponding direction.
+  dynamic clamp(dynamic value, dynamic deltaValue, dynamic hardMin, dynamic hardMax);
+
+  /// Implements looping behaviour for this operation.  Given a [value] that the [deltaValue] will be applied to, return
+  /// a new delta value which won't exceed the clamp constraints by looping the result value from min to max and vice
+  /// versa.
+  /// If the value can't be looped, return null.
+  dynamic loop(dynamic value, dynamic deltaValue, dynamic hardMin, dynamic hardMax);
+
+  /// Returns the parameters for a WebSocket message that will apply the given [modifyEvent] in the engine using this
+  /// operation.
+  dynamic makeMessageParameters(PropertyModifyEvent modifyEvent);
+}
+
+/// A generic arithmetic operation.
+abstract class GenericArithmeticModifyOperation extends PropertyModifyOperation {
+  const GenericArithmeticModifyOperation() : super();
+
+  /// The name of the operation type to send to the engine.
+  String get operationName;
+
+  /// Returns the delta value that, when applied with this operation, would change [oldValue] into [newValue].
+  dynamic makeDeltaValue(dynamic oldValue, dynamic newValue);
+
+  @override
+  String get messageName => 'preset.property.modify';
+
+  @override
+  dynamic apply(dynamic value, dynamic deltaValue) => value + deltaValue;
+
+  @override
+  dynamic clamp(dynamic value, dynamic deltaValue, dynamic hardMin, dynamic hardMax) {
+    if (hardMin == null && hardMax == null) {
+      // No clamping needed
+      return deltaValue;
+    }
+
+    if (hardMin != null && hardMax != null && hardMin > hardMax) {
+      throw Exception('Cannot clamp when min > max ($hardMin > $hardMax)');
+    }
+
+    final dynamic changedValue = apply(value, deltaValue);
+
+    // Determine what the changed value should be after clamping
+    dynamic targetValue;
+    if (hardMax != null && changedValue > hardMax) {
+      targetValue = hardMax;
+    } else if (hardMin != null && changedValue < hardMin) {
+      targetValue = hardMin;
+    } else {
+      // No clamping needed
+      return deltaValue;
+    }
+
+    // If the target value is the same as the original value, then we were already at the clamp limit and this delta
+    // is a no-op.
+    if (targetValue == value) {
+      return null;
+    }
+
+    return targetValue - value;
+  }
+
+  @override
+  dynamic loop(value, deltaValue, hardMin, hardMax) {
+    if (hardMin == null || hardMax == null) {
+      // Can't loop without a min and max
+      return null;
+    }
+
+    if (hardMin >= hardMax) {
+      // Can't loop when min >= max
+      return null;
+    }
+
+    final dynamic range = hardMax - hardMin;
+    dynamic targetValue = apply(value, deltaValue);
+
+    while (targetValue > hardMax) {
+      targetValue -= range;
+      deltaValue -= range;
+    }
+
+    while (targetValue < hardMin) {
+      targetValue += range;
+      deltaValue += range;
+    }
+
+    assert(targetValue >= hardMin && targetValue <= hardMax);
+
+    return targetValue - value;
+  }
+
+  @override
+  dynamic makeMessageParameters(PropertyModifyEvent modifyEvent) {
+    return {
+      'PropertyValue': convertUnrealTypeDartToJson(modifyEvent.typeName, modifyEvent.deltaValue),
+      'PresetName': modifyEvent.presetName,
+      'PropertyLabel': modifyEvent.presetPropertyLabel,
+      'Operation': operationName,
+      'SequenceNumber': modifyEvent.sequenceNumber,
+      'TransactionMode': 'MANUAL',
+      'TransactionId': modifyEvent.transactionId,
+    };
+  }
+}
+
+/// An additive operation on a generic value.
+class AddOperation extends GenericArithmeticModifyOperation {
+  const AddOperation() : super();
+
+  @override
+  makeDeltaValue(oldValue, newValue) => newValue - oldValue;
+
+  @override
+  String get operationName => 'ADD';
+}
+
+/// A multiplicative operation on a generic value.
+class MultiplyOperation extends GenericArithmeticModifyOperation {
+  const MultiplyOperation() : super();
+
+  @override
+  makeDeltaValue(oldValue, newValue) => (oldValue != 0.0) ? (newValue / oldValue) : null;
+
+  @override
+  String get operationName => 'MULTIPLY';
+}
+
+/// An operation that resets the value to its default. The engine doesn't expose the default value to the app, so this
+/// operation doesn't actually change the value locally; instead, we'll get the update from the engine later.
+class ResetOperation extends PropertyModifyOperation {
+  const ResetOperation() : super();
+
+  @override
+  bool get bIsAlwaysSignificant => true;
+
+  @override
+  String get messageName => 'preset.property.modify';
+
+  @override
+  apply(value, deltaValue) => value;
+
+  @override
+  clamp(value, deltaValue, hardMin, hardMax) => deltaValue;
+
+  @override
+  loop(value, deltaValue, hardMin, hardMax) => null;
+
+  @override
+  makeMessageParameters(PropertyModifyEvent modifyEvent) {
+    return {
+      'PropertyValue': convertUnrealTypeDartToJson(modifyEvent.typeName, modifyEvent.deltaValue),
+      'PresetName': modifyEvent.presetName,
+      'PropertyLabel': modifyEvent.presetPropertyLabel,
+      'SequenceNumber': modifyEvent.sequenceNumber,
+      'ResetToDefault': true,
+      'TransactionMode': 'MANUAL',
+      'TransactionId': modifyEvent.transactionId,
+    };
+  }
+}
+
+/// An operation that sets a value directly. deltaValue is treated as an absolute value instead of a relative one.
+/// Note that no clamping is applied.
+class SetOperation extends PropertyModifyOperation {
+  const SetOperation() : super();
+
+  @override
+  bool get bIsAlwaysSignificant => true;
+
+  @override
+  String get messageName => 'preset.property.modify';
+
+  @override
+  apply(value, deltaValue) => deltaValue;
+
+  @override
+  clamp(value, deltaValue, hardMin, hardMax) => deltaValue;
+
+  @override
+  loop(value, deltaValue, hardMin, hardMax) => null;
+
+  @override
+  makeMessageParameters(PropertyModifyEvent modifyEvent) {
+    return {
+      'PropertyValue': convertUnrealTypeDartToJson(modifyEvent.typeName, modifyEvent.deltaValue),
+      'PresetName': modifyEvent.presetName,
+      'PropertyLabel': modifyEvent.presetPropertyLabel,
+      'SequenceNumber': modifyEvent.sequenceNumber,
+      'TransactionMode': 'MANUAL',
+      'TransactionId': modifyEvent.transactionId,
+    };
+  }
+}
+
+/// An additive operation on a [WheelColor]. The value will be clamped within the range of the color wheel.
+class WheelColorAddOperation extends PropertyModifyOperation {
+  const WheelColorAddOperation() : super();
+
+  @override
+  String get messageName => 'object.call';
+
+  @override
+  dynamic apply(dynamic value, dynamic deltaValue) => value + deltaValue;
+
+  @override
+  dynamic clamp(dynamic value, dynamic deltaValue, dynamic hardMin, dynamic hardMax) {
+    // Clamp the value and alpha components in the valid [0,1] range
+    final double clampedDeltaValueComponent =
+        _clampAdditiveDelta(value.value, deltaValue.value, min: 0.0, max: 1.0).toDouble();
+    final double clampedDeltaAlpha = _clampAdditiveDelta(value.alpha, deltaValue.alpha, min: 0.0, max: 1.0).toDouble();
+
+    // Clamp position to a radius of 1
+    final Offset resultPosition = value.position + deltaValue.position;
+    Offset clampedDeltaPosition = deltaValue.position;
+
+    if (resultPosition.distanceSquared > 1) {
+      final Offset clampedPosition = resultPosition / resultPosition.distance;
+      clampedDeltaPosition = clampedPosition - value.position;
+    }
+
+    // Create the new delta value, but return null if it's a no-op
+    final newDeltaValue = WheelColor(clampedDeltaPosition, clampedDeltaValueComponent, clampedDeltaAlpha);
+    if (newDeltaValue == WheelColor.zero) {
+      return null;
+    }
+
+    return newDeltaValue;
+  }
+
+  @override
+  loop(value, deltaValue, hardMin, hardMax) => null;
+
+  @override
+  dynamic makeMessageParameters(PropertyModifyEvent modifyEvent) {
+    return {
+      'FunctionName': 'ApplyColorWheelDelta',
+      'ObjectPath': '/Script/RemoteControl.Default__RemoteControlFunctionLibrary',
+      'Parameters': {
+        'TargetObject': modifyEvent.property.objectPath,
+        'PropertyName': modifyEvent.property.propertyName,
+        'DeltaValue': convertUnrealTypeDartToJson(modifyEvent.typeName, modifyEvent.deltaValue),
+        'ReferenceColor': convertUnrealTypeDartToJson(modifyEvent.typeName, modifyEvent.currentValue),
+        'bIsInteractive': true,
+      },
+      'SequenceNumber': modifyEvent.sequenceNumber,
+      'TransactionMode': 'MANUAL',
+      'TransactionId': modifyEvent.transactionId,
+    };
+  }
+}
+
+/// An additive operation on a [WheelGradingColor]. The value will be clamped within the range of the color wheel.
+class WheelColorGradingAddOperation extends PropertyModifyOperation {
+  const WheelColorGradingAddOperation({
+    required this.saturationExponent,
+    required this.minValue,
+    required this.maxValue,
+  }) : super();
+
+  /// Exponent applied to the saturation component of the wheel, making it less sensitive towards the center.
+  final double saturationExponent;
+
+  /// The minimum value of the slider. The value will be clamped in this range and scaled based on the min/max range
+  /// when received in the engine.
+  final double minValue;
+
+  /// The minimum value of the slider. The value will be clamped in this range and scaled based on the min/max range
+  /// when received in the engine.
+  final double maxValue;
+
+  @override
+  String get messageName => 'object.call';
+
+  @override
+  dynamic apply(dynamic value, dynamic deltaValue) => value + deltaValue;
+
+  @override
+  dynamic clamp(dynamic value, dynamic deltaValue, dynamic hardMin, dynamic hardMax) {
+    // Clamp the value and alpha components in the valid range
+    final double clampedDeltaValueComponent =
+        _clampAdditiveDelta(value.value, deltaValue.value, min: minValue, max: maxValue).toDouble();
+    final double clampedDeltaAlpha =
+        _clampAdditiveDelta(value.alpha, deltaValue.alpha, min: minValue, max: maxValue).toDouble();
+
+    // Delta is applied by user in exponential space, so find the delta in linear space
+    final Offset exponentialPosition = exponentiateCirclePosition(value.position, saturationExponent, true);
+    final Offset exponentialResultPosition = exponentialPosition + deltaValue.position;
+    final Offset linearResultPosition =
+        exponentiateCirclePosition(exponentialResultPosition, saturationExponent, false);
+    final Offset deltaPosition = linearResultPosition - value.position;
+
+    // Clamp position to a radius of 1 (this is unaffected by min/maxValue since the position represents saturation,
+    // which is always in range [0, 1])
+    Offset clampedDeltaPosition = deltaPosition;
+
+    if (linearResultPosition.distanceSquared > 1) {
+      final Offset clampedPosition = linearResultPosition / linearResultPosition.distance;
+      clampedDeltaPosition = clampedPosition - value.position;
+    }
+
+    // Create the new delta value, but return null if it's a no-op
+    final newDeltaValue = WheelColor(clampedDeltaPosition, clampedDeltaValueComponent, clampedDeltaAlpha);
+    if (newDeltaValue == WheelColor.zero) {
+      return null;
+    }
+
+    return newDeltaValue;
+  }
+
+  @override
+  loop(value, deltaValue, hardMin, hardMax) => null;
+
+  @override
+  dynamic makeMessageParameters(PropertyModifyEvent modifyEvent) {
+    return {
+      'FunctionName': 'ApplyColorGradingWheelDelta',
+      'ObjectPath': '/Script/RemoteControl.Default__RemoteControlFunctionLibrary',
+      'Parameters': {
+        'TargetObject': modifyEvent.property.objectPath,
+        'PropertyName': modifyEvent.property.propertyName,
+        'DeltaValue': modifyEvent.deltaValue.toJson(bUseLuminance: true),
+        'ReferenceColor': modifyEvent.currentValue.toJson(bUseLuminance: true),
+        'bIsInteractive': true,
+        'MinValue': minValue,
+        'MaxValue': maxValue,
+      },
+      'SequenceNumber': modifyEvent.sequenceNumber,
+      'TransactionMode': 'MANUAL',
+      'TransactionId': modifyEvent.transactionId,
+    };
+  }
+}
+
+/// Given a base [value] and a [deltaValue] to add to it, return the clamped delta such that the final sum won't exceed
+/// the [min] or [max] values.
+num _clampAdditiveDelta(num value, num deltaValue, {num? min, num? max}) {
+  num clampedResult = value + deltaValue;
+
+  if (min != null && clampedResult < min) {
+    clampedResult = min;
+  } else if (max != null && clampedResult > max) {
+    clampedResult = max;
+  } else {
+    return deltaValue;
+  }
+
+  return clampedResult - value;
+}
diff --git a/Engine/Extras/VirtualProduction/EpicStageApp/lib/models/settings/base_color_tab_settings.dart b/Engine/Extras/VirtualProduction/EpicStageApp/lib/models/settings/base_color_tab_settings.dart
new file mode 100644
--- /dev/null
+++ b/Engine/Extras/VirtualProduction/EpicStageApp/lib/models/settings/base_color_tab_settings.dart
@@ -0,0 +1,42 @@
+// Copyright Epic Games, Inc. All Rights Reserved.
+
+import 'package:streaming_shared_preferences/streaming_shared_preferences.dart';
+
+import '../../utilities/preferences_bundle.dart';
+import '../../utilities/streaming_preferences_extensions.dart';
+import 'color_wheel_settings.dart';
+
+/// Types of limits the user can edit for a color range.
+enum RangeLimitType {
+  min,
+  max,
+}
+
+/// User settings for [BaseColorTab].
+class BaseColorTabSettings {
+  BaseColorTabSettings(PreferencesBundle preferences)
+      : colorGradingRange = preferences.persistent.getEnum(
+          'colorTab.colorGradingRange',
+          defaultValue: ColorGradingRange.global,
+          enumValues: ColorGradingRange.values,
+        ),
+        colorGradingSubproperty = preferences.persistent.getEnum(
+          'colorTab.colorGradingSubproperty',
+          defaultValue: ColorGradingSubproperty.gain,
+          enumValues: ColorGradingSubproperty.values,
+        ),
+        colorGradingHighlightsRangeLimitType = preferences.persistent.getEnum(
+          'colorTab.colorGradingHighlightsRangeLimitType',
+          defaultValue: RangeLimitType.min,
+          enumValues: RangeLimitType.values,
+        );
+
+  /// The last color grading range selected for color grading control.
+  final Preference<ColorGradingRange> colorGradingRange;
+
+  /// The last color grading subproperty selected for color grading control.
+  final Preference<ColorGradingSubproperty> colorGradingSubproperty;
+
+  /// Which type of highlights range slider to show when color grading.
+  final Preference<RangeLimitType> colorGradingHighlightsRangeLimitType;
+}
diff --git a/Engine/Extras/VirtualProduction/EpicStageApp/lib/models/settings/color_grading_tab_settings.dart b/Engine/Extras/VirtualProduction/EpicStageApp/lib/models/settings/color_grading_tab_settings.dart
new file mode 100644
--- /dev/null
+++ b/Engine/Extras/VirtualProduction/EpicStageApp/lib/models/settings/color_grading_tab_settings.dart
@@ -0,0 +1,26 @@
+// Copyright Epic Games, Inc. All Rights Reserved.
+
+import 'package:streaming_shared_preferences/streaming_shared_preferences.dart';
+
+import '../../utilities/preferences_bundle.dart';
+import '../../utilities/transient_preference.dart';
+import '../unreal_types.dart';
+
+/// Holds user settings used in the Color Grading tab.
+class ColorGradingTabSettings {
+  ColorGradingTabSettings(PreferencesBundle preferences)
+      : target = preferences.transient.get(
+          'colorGradingTab.target',
+          defaultValue: UnrealProperty.empty,
+        ),
+        panelHeight = preferences.persistent.getDouble(
+          'colorGradingTab.panelHeight',
+          defaultValue: -1,
+        );
+
+  /// The property selected as the color grading target.
+  final TransientPreference<UnrealProperty> target;
+
+  /// How tall the color grading target panel should be.
+  final Preference<double> panelHeight;
+}
diff --git a/Engine/Extras/VirtualProduction/EpicStageApp/lib/models/settings/color_wheel_settings.dart b/Engine/Extras/VirtualProduction/EpicStageApp/lib/models/settings/color_wheel_settings.dart
new file mode 100644
--- /dev/null
+++ b/Engine/Extras/VirtualProduction/EpicStageApp/lib/models/settings/color_wheel_settings.dart
@@ -0,0 +1,56 @@
+// Copyright Epic Games, Inc. All Rights Reserved.
+
+import 'package:streaming_shared_preferences/streaming_shared_preferences.dart';
+
+import '../../utilities/preferences_bundle.dart';
+import '../../utilities/streaming_preferences_extensions.dart';
+
+/// What type of color editor to display.
+enum ColorWheelEditMode {
+  wheel,
+  sliders,
+}
+
+/// Which color components to display.
+enum ColorWheelColorMode {
+  rgb,
+  hsv,
+}
+
+/// Color ranges that the user can control separately.
+enum ColorGradingRange {
+  global,
+  shadows,
+  midtones,
+  highlights,
+}
+
+/// Color grading color properties that the user can control in each range.
+enum ColorGradingSubproperty {
+  saturation,
+  contrast,
+  gamma,
+  gain,
+  offset,
+}
+
+/// Settings that the user can change on color wheels.
+class ColorWheelSettings {
+  ColorWheelSettings(PreferencesBundle preferences)
+      : editMode = preferences.persistent.getEnum(
+          'colorWheel.editMode',
+          defaultValue: ColorWheelEditMode.wheel,
+          enumValues: ColorWheelEditMode.values,
+        ),
+        colorMode = preferences.persistent.getEnum(
+          'colorWheel.colorMode',
+          defaultValue: ColorWheelColorMode.rgb,
+          enumValues: ColorWheelColorMode.values,
+        );
+
+  /// What type of color editor to display.
+  final Preference<ColorWheelEditMode> editMode;
+
+  /// Which color components to display.
+  final Preference<ColorWheelColorMode> colorMode;
+}
diff --git a/Engine/Extras/VirtualProduction/EpicStageApp/lib/models/settings/delta_widget_settings.dart b/Engine/Extras/VirtualProduction/EpicStageApp/lib/models/settings/delta_widget_settings.dart
new file mode 100644
--- /dev/null
+++ b/Engine/Extras/VirtualProduction/EpicStageApp/lib/models/settings/delta_widget_settings.dart
@@ -0,0 +1,86 @@
+// Copyright Epic Games, Inc. All Rights Reserved.
+
+import 'package:flutter/material.dart';
+import 'package:flutter_gen/gen_l10n/app_localizations.dart';
+import 'package:logging/logging.dart';
+import 'package:streaming_shared_preferences/streaming_shared_preferences.dart';
+
+import '../../utilities/debug_utilities.dart';
+import '../../utilities/preferences_bundle.dart';
+import '../../utilities/transient_preference.dart';
+import '../navigator_keys.dart';
+
+final _log = Logger('DeltaWidgetSettings');
+
+/// Types of limits the user can edit for a color range.
+enum RangeLimitType {
+  min,
+  max,
+}
+
+/// Which type of actor properties to display in the details panel.
+enum DetailsPropertyDisplayType {
+  appearance,
+  orientation,
+}
+
+/// User settings for delta-based widgets.
+class DeltaWidgetSettings {
+  DeltaWidgetSettings(PreferencesBundle preferences)
+      : bIsInResetMode = preferences.transient.get('colorTab.bIsInResetMode', defaultValue: false),
+        sensitivity = preferences.persistent.getDouble('deltaWidgets.sensitivity', defaultValue: 1.0) {
+    bIsInResetMode.listen((_) => _validateResetMode());
+  }
+
+  /// Whether the Unreal widgets should be displayed in reset mode.
+  final TransientPreference<bool> bIsInResetMode;
+
+  /// The base sensitivity of all delta-based widgets.
+  final Preference<double> sensitivity;
+
+  /// Map from owners of reset mode blocks to the message to display when blocked.
+  final Map<Object, String> _resetModeBlockers = {};
+
+  /// Check whether reset mode is blocked.
+  bool get bIsResetModeBlocked => _resetModeBlockers.length > 0;
+
+  /// Add a blocker for reset mode. The [owner] is responsible for removing it when appropriate.
+  /// When a blocker is present and reset mode is on, we will force it off and show the user a [message] explaining why.
+  void addResetModeBlocker(Object owner, String message) {
+    if (_resetModeBlockers.containsKey(owner)) {
+      _log.warning('$owner tried to add a reset mode block, but it already had one');
+      return;
+    }
+
+    _resetModeBlockers[owner] = message;
+
+    _validateResetMode();
+  }
+
+  /// Remove a blocker for reset mode.
+  void removeResetModeBlocker(Object owner) {
+    _resetModeBlockers.remove(owner);
+  }
+
+  /// Check whether [owner] owns a reset mode blocker which is currently active.
+  bool isResetModeBlockedBy(Object owner) {
+    return _resetModeBlockers.containsKey(owner);
+  }
+
+  /// Check whether reset mode can be set, and if not, turn it off.
+  void _validateResetMode() {
+    if (!bIsInResetMode.getValue() || !bIsResetModeBlocked) {
+      return;
+    }
+
+    final BuildContext rootContext = rootNavigatorKey.currentContext!;
+    final String combinedMessage =
+        '${AppLocalizations.of(rootContext)!.resetModeDisabledMessage}\n\n' + _resetModeBlockers.values.join('\n\n');
+
+    showDebugAlert(combinedMessage);
+
+    WidgetsBinding.instance.addPostFrameCallback((_) {
+      bIsInResetMode.setValue(false);
+    });
+  }
+}
diff --git a/Engine/Extras/VirtualProduction/EpicStageApp/lib/models/settings/details_tab_settings.dart b/Engine/Extras/VirtualProduction/EpicStageApp/lib/models/settings/details_tab_settings.dart
new file mode 100644
--- /dev/null
+++ b/Engine/Extras/VirtualProduction/EpicStageApp/lib/models/settings/details_tab_settings.dart
@@ -0,0 +1,25 @@
+// Copyright Epic Games, Inc. All Rights Reserved.
+
+import 'package:streaming_shared_preferences/streaming_shared_preferences.dart';
+
+import '../../utilities/preferences_bundle.dart';
+import '../../utilities/streaming_preferences_extensions.dart';
+
+/// Which type of actor properties to display in the details panel.
+enum DetailsPropertyDisplayType {
+  appearance,
+  orientation,
+}
+
+/// Holds user settings used in the Details tab.
+class DetailsTabSettings {
+  DetailsTabSettings(PreferencesBundle preferences)
+      : detailsPropertyDisplayType = preferences.persistent.getEnum(
+          'colorGradingTab.detailsPropertyDisplayType',
+          defaultValue: DetailsPropertyDisplayType.appearance,
+          enumValues: DetailsPropertyDisplayType.values,
+        );
+
+  /// Which type of actor properties to display in the details panel.
+  final Preference<DetailsPropertyDisplayType> detailsPropertyDisplayType;
+}
diff --git a/Engine/Extras/VirtualProduction/EpicStageApp/lib/models/settings/floating_map_settings.dart b/Engine/Extras/VirtualProduction/EpicStageApp/lib/models/settings/floating_map_settings.dart
new file mode 100644
--- /dev/null
+++ b/Engine/Extras/VirtualProduction/EpicStageApp/lib/models/settings/floating_map_settings.dart
@@ -0,0 +1,39 @@
+// Copyright Epic Games, Inc. All Rights Reserved.
+
+import 'package:streaming_shared_preferences/streaming_shared_preferences.dart';
+
+import '../../utilities/preferences_bundle.dart';
+import '../../utilities/streaming_preferences_extensions.dart';
+
+/// Which side of the screen the map preview is on.
+enum PreviewMapSide {
+  left,
+  right,
+}
+
+/// Holds user settings used by the floating preview map.
+class FloatingMapSettings {
+  FloatingMapSettings(PreferencesBundle preferences)
+      : mapSide = preferences.persistent.getEnum<PreviewMapSide>(
+          'floatingMapSettings.mapSide',
+          defaultValue: PreviewMapSide.right,
+          enumValues: PreviewMapSide.values,
+        ),
+        mapY = preferences.persistent.getDouble(
+          'floatingMapSettings.mapY',
+          defaultValue: 1.0,
+        ),
+        bIsMapDocked = preferences.persistent.getBool(
+          'floatingMapSettings.bIsMapDocked',
+          defaultValue: false,
+        );
+
+  /// Which side of the screen the preview map is on.
+  Preference<PreviewMapSide> mapSide;
+
+  /// The vertical position at which the map preview should sit (as a fraction of its container height).
+  Preference<double> mapY;
+
+  /// Whether the preview map is docked off-screen.
+  Preference<bool> bIsMapDocked;
+}
diff --git a/Engine/Extras/VirtualProduction/EpicStageApp/lib/models/settings/main_screen_settings.dart b/Engine/Extras/VirtualProduction/EpicStageApp/lib/models/settings/main_screen_settings.dart
new file mode 100644
--- /dev/null
+++ b/Engine/Extras/VirtualProduction/EpicStageApp/lib/models/settings/main_screen_settings.dart
@@ -0,0 +1,24 @@
+// Copyright Epic Games, Inc. All Rights Reserved.
+
+import 'package:streaming_shared_preferences/streaming_shared_preferences.dart';
+
+import '../../utilities/preferences_bundle.dart';
+
+/// Holds settings that affect the entire layout of the main screen.
+class MainScreenSettings {
+  MainScreenSettings(PreferencesBundle preferences)
+      : bIsOutlinerPanelOpen = preferences.persistent.getBool(
+          'common.bIsOutlinerPanelOpen',
+          defaultValue: false,
+        ),
+        selectedTab = preferences.persistent.getInt(
+          'common.selectedTab',
+          defaultValue: 0,
+        );
+
+  /// Whether the outliner panel is open.
+  final Preference<bool> bIsOutlinerPanelOpen;
+
+  /// The selected top-level tab index.
+  final Preference<int> selectedTab;
+}
diff --git a/Engine/Extras/VirtualProduction/EpicStageApp/lib/models/settings/outliner_panel_settings.dart b/Engine/Extras/VirtualProduction/EpicStageApp/lib/models/settings/outliner_panel_settings.dart
new file mode 100644
--- /dev/null
+++ b/Engine/Extras/VirtualProduction/EpicStageApp/lib/models/settings/outliner_panel_settings.dart
@@ -0,0 +1,40 @@
+// Copyright Epic Games, Inc. All Rights Reserved.
+
+import 'package:streaming_shared_preferences/streaming_shared_preferences.dart';
+
+import '../../utilities/preferences_bundle.dart';
+import '../../utilities/streaming_preferences_extensions.dart';
+import '../unreal_object_filters.dart';
+
+/// Ways the user can sort the list of actors in the outliner panel.
+enum OutlinerActorSortMode {
+  name,
+  recent,
+}
+
+/// Holds user settings used to render and control the stage map.
+class OutlinerPanelSettings {
+  OutlinerPanelSettings(PreferencesBundle preferences)
+      : bShouldSortActorsDescending = preferences.persistent.getBool(
+          'outliner.bShouldSortActorsDescending',
+          defaultValue: true,
+        ),
+        actorSortMode = preferences.persistent.getEnum<OutlinerActorSortMode>(
+          'outliner.actorSortMode',
+          defaultValue: OutlinerActorSortMode.name,
+          enumValues: OutlinerActorSortMode.values,
+        ),
+        selectedFilters = preferences.persistent.getUnrealObjectFilterSet(
+          'outliner.filters',
+          defaultValue: {},
+        );
+
+  /// Whether to sort actors in descending order.
+  final Preference<bool> bShouldSortActorsDescending;
+
+  /// The metric by which to sort the actors.
+  final Preference<OutlinerActorSortMode> actorSortMode;
+
+  /// Which actor filters are selected.
+  final Preference<Set<UnrealObjectFilter>> selectedFilters;
+}
diff --git a/Engine/Extras/VirtualProduction/EpicStageApp/lib/models/settings/recent_actor_settings.dart b/Engine/Extras/VirtualProduction/EpicStageApp/lib/models/settings/recent_actor_settings.dart
new file mode 100644
--- /dev/null
+++ b/Engine/Extras/VirtualProduction/EpicStageApp/lib/models/settings/recent_actor_settings.dart
@@ -0,0 +1,82 @@
+// Copyright Epic Games, Inc. All Rights Reserved.
+
+import 'package:streaming_shared_preferences/streaming_shared_preferences.dart';
+
+import '../../utilities/json_utilities.dart';
+import '../../utilities/preferences_bundle.dart';
+import '../../utilities/streaming_preferences_extensions.dart';
+
+enum RecentlyPlacedActorType {
+  lightCard,
+  chromakeyCard,
+  flag,
+  colorCorrectWindow,
+}
+
+/// Data about an actor the user recently placed.
+class RecentlyPlacedActorData {
+  const RecentlyPlacedActorData({this.name, required this.type, this.templatePath});
+
+  RecentlyPlacedActorData.fromJson(Map<String, dynamic> json)
+      : name = json['Name'],
+        type = jsonToEnumValue(json['Type'], RecentlyPlacedActorType.values) ?? RecentlyPlacedActorType.lightCard,
+        templatePath = json['TemplatePath'];
+
+  /// The name to show for the actor. If null, use the default name for the [type].
+  final String? name;
+
+  /// The type of actor to create.
+  final RecentlyPlacedActorType type;
+
+  /// The path to the template used to create the actor.
+  final String? templatePath;
+
+  @override
+  bool operator ==(Object other) {
+    if (identical(this, other)) {
+      return true;
+    }
+
+    if (other.runtimeType != runtimeType) {
+      return false;
+    }
+
+    return other is RecentlyPlacedActorData && other.type == type && other.templatePath == templatePath;
+  }
+
+  Map<String, dynamic> toJson() {
+    return {
+      'Name': name,
+      'Type': enumToJsonValue(type),
+      'TemplatePath': templatePath,
+    };
+  }
+}
+
+/// Holds settings about actors that the user recently placed.
+class RecentActorSettings {
+  RecentActorSettings(PreferencesBundle preferences)
+      : recentlyPlacedActors = preferences.persistent.getRecentlyPlacedActorList(
+          'common.recentlyPlacedActors',
+          defaultValue: [],
+        );
+
+  /// The last time the actor with the given path was selected.
+  final Preference<List<RecentlyPlacedActorData>> recentlyPlacedActors;
+
+  /// Add a recently placed actor (or move it to the top of the list if it's already in the list).
+  void addRecentlyPlacedActor(RecentlyPlacedActorData actorData) {
+    const int historyLength = 5;
+
+    final List<RecentlyPlacedActorData> actorList = recentlyPlacedActors.getValue();
+
+    actorList.remove(actorData);
+    actorList.insert(0, actorData);
+
+    while (actorList.length > historyLength) {
+      actorList.removeLast();
+    }
+
+    recentlyPlacedActors.setValue(actorList);
+  }
+}
diff --git a/Engine/Extras/VirtualProduction/EpicStageApp/lib/models/settings/selected_actor_settings.dart b/Engine/Extras/VirtualProduction/EpicStageApp/lib/models/settings/selected_actor_settings.dart
new file mode 100644
--- /dev/null
+++ b/Engine/Extras/VirtualProduction/EpicStageApp/lib/models/settings/selected_actor_settings.dart
@@ -0,0 +1,124 @@
+// Copyright Epic Games, Inc. All Rights Reserved.
+
+import 'package:streaming_shared_preferences/streaming_shared_preferences.dart';
+
+import '../../utilities/preferences_bundle.dart';
+import '../../utilities/transient_preference.dart';
+
+/// Holds user settings that impact the state of the selected actors.
+class SelectedActorSettings {
+  SelectedActorSettings(PreferencesBundle preferences)
+      : displayClusterRootPath = preferences.persistent.getString('common.displayClusterRootPath', defaultValue: ''),
+        selectedActors = preferences.transient.get<Set<String>>(
+          'common.selectedActors',
+          defaultValue: {},
+        ),
+        lastSelectedTimeForActors = preferences.transient.get<Map<String, DateTime>>(
+          'common.lastSelectedTimeForActors',
+          defaultValue: {},
+        ),
+        bIsInMultiSelectMode = preferences.persistent.getBool(
+          'common.bIsInMultiSelectMode',
+          defaultValue: false,
+        ) {
+    bIsInMultiSelectMode.listen(_onIsInMultiSelectModeChanged);
+    displayClusterRootPath.listen(_onDisplayClusterRootPathChanged);
+  }
+
+  /// One-off serialization adapter for a [Map<String, DateTime>].
+  static PreferenceAdapter<Map<String, DateTime>> stringDateTimeMapAdapter = JsonAdapter(
+    serializer: (Map<String, DateTime> map) => map.map(
+      (String key, DateTime value) => MapEntry(key, value.millisecondsSinceEpoch.toString()),
+    ),
+    deserializer: (jsonData) {
+      if (!(jsonData is Map<String, String>)) {
+        return {};
+      }
+
+      return jsonData.map(
+        (key, value) => MapEntry(key, DateTime.fromMillisecondsSinceEpoch(int.parse(value), isUtc: true)),
+      );
+    },
+  );
+
+  /// The path of the selected nDisplay cluster.
+  final Preference<String> displayClusterRootPath;
+
+  /// Paths of actors that the user has selected to multi-edit.
+  final TransientPreference<Set<String>> selectedActors;
+
+  /// The last time the actor with the given path was selected.
+  final TransientPreference<Map<String, DateTime>> lastSelectedTimeForActors;
+
+  /// If true, the user can select multiple light cards at once. Otherwise, only the most recently selected light card
+  /// will be selected.
+  final Preference<bool> bIsInMultiSelectMode;
+
+  /// Convenience function to check if an actor with the given [path] is selected.
+  bool isActorSelected(String path) {
+    return selectedActors.getValue().contains(path);
+  }
+
+  /// Select or deselect an actor for multi-edit.
+  void selectActor(String path, {bool bShouldSelect = true}) {
+    final Set<String> selectedActorSet = selectedActors.getValue().toSet();
+    bool bChanged = false;
+
+    if (bShouldSelect) {
+      if (!bIsInMultiSelectMode.getValue()) {
+        selectedActorSet.clear();
+        bChanged = true;
+      }
+
+      if (selectedActorSet.add(path)) {
+        bChanged = true;
+      }
+    } else if (selectedActorSet.remove(path)) {
+      // The actor was deselected, so remember that it was selected until just now
+      final newLastSelectedTimeForActors = lastSelectedTimeForActors.getValue();
+      newLastSelectedTimeForActors[path] = DateTime.now();
+      lastSelectedTimeForActors.setValue(newLastSelectedTimeForActors);
+
+      bChanged = true;
+    }
+
+    if (bChanged) {
+      selectedActors.setValue(selectedActorSet);
+    }
+  }
+
+  /// Deselect all actors that are currently selected.
+  void deselectAllActors() {
+    selectedActors.setValue({});
+  }
+
+  /// Get the last time an actor was still selected.
+  /// If the actor was selected but no longer is, return the time at which it became deselected.
+  /// If the actor is currently selected, return the current time.
+  /// If the actor was never selected, return null.
+  DateTime? getActorLastSelectedTime(String path) {
+    if (isActorSelected(path)) {
+      return DateTime.now();
+    }
+
+    return lastSelectedTimeForActors.getValue()[path];
+  }
+
+  /// Called when the value of [bIsInMultiSelectMode] changes.
+  void _onIsInMultiSelectModeChanged(bool bNewIsInMultiSelectMode) {
+    if (bNewIsInMultiSelectMode) {
+      return;
+    }
+
+    final selectedActorSet = selectedActors.getValue();
+    if (selectedActorSet.length > 0) {
+      // Reduce to one actor. Set is ordered, so this should give us the last selected actor.
+      selectedActors.setValue({selectedActorSet.last});
+    }
+  }
+
+  /// Called when the value of [displayClusterRootPath] changes.
+  void _onDisplayClusterRootPathChanged(_) {
+    deselectAllActors();
+  }
+}
diff --git a/Engine/Extras/VirtualProduction/EpicStageApp/lib/models/settings/stage_map_settings.dart b/Engine/Extras/VirtualProduction/EpicStageApp/lib/models/settings/stage_map_settings.dart
new file mode 100644
--- /dev/null
+++ b/Engine/Extras/VirtualProduction/EpicStageApp/lib/models/settings/stage_map_settings.dart
@@ -0,0 +1,60 @@
+// Copyright Epic Games, Inc. All Rights Reserved.
+
+import 'dart:ui';
+
+import 'package:streaming_shared_preferences/streaming_shared_preferences.dart';
+import 'package:vector_math/vector_math_64.dart' as vec;
+
+import '../../utilities/preferences_bundle.dart';
+import '../../utilities/streaming_preferences_extensions.dart';
+import '../../utilities/transient_preference.dart';
+
+enum ProjectionMode {
+  azimuthal,
+  orthographic,
+  perspective,
+  uv,
+}
+
+/// Holds user settings used to render and control the stage map.
+class StageMapSettings {
+  StageMapSettings(PreferencesBundle preferences)
+      : bIsInTransformMode = preferences.persistent.getBool(
+          'stageMap.bIsInMultiSelectMode',
+          defaultValue: false,
+        ),
+        projectionMode = preferences.persistent.getEnum<ProjectionMode>(
+          'stageMap.projectionMode',
+          defaultValue: ProjectionMode.azimuthal,
+          enumValues: ProjectionMode.values,
+        ),
+        cameraAngle = preferences.transient.get(
+          'stageMap.cameraAngle',
+          defaultValue: vec.Vector2(0, 90),
+        ),
+        focalPoint = preferences.transient.get(
+          'stageMap.focalPoint',
+          defaultValue: Offset(0.5, 0.5),
+        ),
+        zoomLevel = preferences.transient.get(
+          'stageMap.zoomLevel',
+          defaultValue: 1.0,
+        );
+
+  /// If true, the map is in transform mode and the user can pinch/pan to adjust lightcard scale/rotation/location.
+  /// If false, the map is in map mode, and the user can pinch/pan to move the map or long-press pins to drag them
+  /// around.
+  final Preference<bool> bIsInTransformMode;
+
+  /// The camera's current projection mode for stage actor locations.
+  final Preference<ProjectionMode> projectionMode;
+
+  /// The camera's current rotation setting in degrees, where X is yaw and Y is pitch.
+  final TransientPreference<vec.Vector2> cameraAngle;
+
+  /// The location on the map in center of the stage map view, assuming a 1:1 aspect ratio.
+  final TransientPreference<Offset> focalPoint;
+
+  /// The zoom level of the stage map view.
+  final TransientPreference<double> zoomLevel;
+}
diff --git a/Engine/Extras/VirtualProduction/EpicStageApp/lib/models/unreal_actor_creator.dart b/Engine/Extras/VirtualProduction/EpicStageApp/lib/models/unreal_actor_creator.dart
new file mode 100644
--- /dev/null
+++ b/Engine/Extras/VirtualProduction/EpicStageApp/lib/models/unreal_actor_creator.dart
@@ -0,0 +1,350 @@
+// Copyright Epic Games, Inc. All Rights Reserved.
+
+import 'dart:async';
+
+import 'package:flutter/material.dart';
+import 'package:flutter_gen/gen_l10n/app_localizations.dart';
+import 'package:logging/logging.dart';
+import 'package:provider/provider.dart';
+
+import '../utilities/constants.dart';
+import '../utilities/unreal_utilities.dart';
+import 'color.dart';
+import 'engine_connection.dart';
+import 'navigator_keys.dart';
+import 'preview_render_manager.dart';
+import 'settings/recent_actor_settings.dart';
+import 'settings/selected_actor_settings.dart';
+import 'unreal_actor_manager.dart';
+import 'unreal_types.dart';
+
+final _log = Logger('UnrealActorCreator');
+
+/// How long to wait before forgetting about an actor creation request (meaning even if the actor is created, we won't
+/// select it in the app).
+const Duration _creationRequestTimeout = Duration(seconds: 1);
+
+/// Settings that apply when creating any type of actor.
+class ActorCreationSettings {
+  const ActorCreationSettings({
+    this.name,
+    this.template,
+    this.mapPosition,
+    this.callback,
+    this.extraParams,
+  });
+
+  /// The name to give the actor.
+  final String? name;
+
+  /// The template from which to create the actor.
+  final UnrealTemplateData? template;
+
+  /// The normalized position within the preview map at which to place the actor.
+  final Offset? mapPosition;
+
+  /// A function to call once the engine reports that the actor has been created.
+  final void Function()? callback;
+
+  /// Any extra JSON parameters to pass with the creation request.
+  final Map<String, dynamic>? extraParams;
+
+  /// Actor settings left at their default values.
+  static const ActorCreationSettings defaults = ActorCreationSettings();
+
+  /// Make a copy of the settings, replacing certain parameters.
+  ActorCreationSettings copyWith({
+    String? name,
+    UnrealTemplateData? template,
+    Offset? mapPosition,
+    void Function()? callback,
+    Map<String, dynamic>? extraParams,
+  }) {
+    return ActorCreationSettings(
+      name: name ?? this.name,
+      template: template ?? this.template,
+      mapPosition: mapPosition ?? this.mapPosition,
+      callback: callback ?? this.callback,
+      extraParams: extraParams ?? this.extraParams,
+    );
+  }
+}
+
+/// A utility that manages creating actors and handling creation responses from the engine.
+/// This is separate from UnrealActorManager because it also depends on user settings, previews, etc.
+class UnrealActorCreator {
+  UnrealActorCreator(BuildContext context)
+      : _connectionManager = Provider.of<EngineConnectionManager>(context, listen: false),
+        _previewRenderManager = Provider.of<PreviewRenderManager>(context, listen: false),
+        _actorManager = Provider.of<UnrealActorManager>(context, listen: false),
+        _recentActorSettings = Provider.of<RecentActorSettings>(context, listen: false),
+        _selectedActorSettings = Provider.of<SelectedActorSettings>(context, listen: false),
+        _localizations = AppLocalizations.of(rootNavigatorKey.currentContext!)! {
+    _connectionManager..registerMessageListener('RequestedActorsCreated', _onRequestedActorsCreated);
+  }
+
+  final EngineConnectionManager _connectionManager;
+  final UnrealActorManager _actorManager;
+  final PreviewRenderManager _previewRenderManager;
+  final RecentActorSettings _recentActorSettings;
+  final SelectedActorSettings _selectedActorSettings;
+  final AppLocalizations _localizations;
+
+  /// The ID to send with actor creation requests to the engine.
+  int _lastRequestId = 0;
+
+  /// Map from IDs of actor creation requests pending a response from the engine to data about the request.
+  Map<int, _PendingCreationRequest> _pendingRequests = {};
+
+  /// Call this to clean up any subscriptions.
+  void dispose() {
+    _connectionManager.unregisterMessageListener('RequestedActorsCreated', _onRequestedActorsCreated);
+  }
+
+  /// Send a message to the engine requesting that a new lightcard be created.
+  void createLightcard([ActorCreationSettings settings = ActorCreationSettings.defaults]) async {
+    final UnrealTemplateData? passedTemplate = settings.template;
+
+    // Set the default template if there is one. This isn't strictly necessary since the engine should use it anyway,
+    // but we want to explicitly set it so we know that the actor already has its default position provided by the
+    // template.
+    UnrealTemplateData? defaultTemplate;
+    if (settings.template == null) {
+      // Check if there's a default template
+      final UnrealHttpResponse templateResponse = await _connectionManager.sendHttpRequest(UnrealHttpRequest(
+        url: '/remote/object/property',
+        verb: 'PUT',
+        body: {
+          'objectPath': '/Script/DisplayClusterLightCardEditor.Default__DisplayClusterLightCardEditorProjectSettings',
+          'propertyName': 'DefaultLightCardTemplate',
+          'access': 'READ_ACCESS',
+        },
+      ));
+
+      if (templateResponse.code == HttpResponseCode.ok) {
+        final String? defaultTemplatePath = templateResponse.body?['DefaultLightCardTemplate'];
+        if (defaultTemplatePath != null && defaultTemplatePath.isNotEmpty) {
+          defaultTemplate = UnrealTemplateData(name: '', path: defaultTemplatePath);
+        }
+      }
+    }
+
+    settings = settings.copyWith(
+      name: settings.name ?? _localizations.newActorNameLightCard,
+      template: passedTemplate ?? defaultTemplate,
+    );
+
+    createActor(className: lightCardClassName, settings: settings);
+
+    _recentActorSettings.addRecentlyPlacedActor(RecentlyPlacedActorData(
+      name: passedTemplate?.name,
+      type: RecentlyPlacedActorType.lightCard,
+      templatePath: passedTemplate?.path,
+    ));
+  }
+
+  /// Send a message to the engine requesting that a new lightcard be created.
+  void createChromakeyCard([ActorCreationSettings settings = ActorCreationSettings.defaults]) {
+    if (settings.name == null) {
+      settings = settings.copyWith(name: _localizations.newActorNameChromakeyCard);
+    }
+
+    createActor(className: chromakeyCardClassName, settings: settings);
+
+    _recentActorSettings.addRecentlyPlacedActor(RecentlyPlacedActorData(
+      name: settings.template?.name,
+      type: RecentlyPlacedActorType.chromakeyCard,
+      templatePath: settings.template?.path,
+    ));
+  }
+
+  /// Send a message to the engine requesting that a new flag be created.
+  void createFlag([ActorCreationSettings settings = ActorCreationSettings.defaults]) async {
+    // First, try to get the default flag template
+    final UnrealHttpResponse templateResponse = await _connectionManager.sendHttpRequest(UnrealHttpRequest(
+      url: '/remote/object/property',
+      verb: 'PUT',
+      body: {
+        'objectPath': '/Script/DisplayClusterLightCardEditor.Default__DisplayClusterLightCardEditorProjectSettings',
+        'propertyName': 'DefaultFlagTemplate',
+        'access': 'READ_ACCESS',
+      },
+    ));
+
+    // If the template isn't available, explicitly send None so we don't get the default LC template instead.
+    String templatePath = 'None';
+
+    if (templateResponse.code == HttpResponseCode.ok) {
+      final String? defaultTemplate = templateResponse.body?['DefaultFlagTemplate'];
+      if (defaultTemplate != null && defaultTemplate.isNotEmpty) {
+        templatePath = defaultTemplate;
+      }
+    }
+
+    final Map<String, dynamic> extraParams = {
+      'bIsLightCardFlag': true,
+      'OverrideColor': true,
+      'Color': FloatColor(0, 0, 0).toJson(),
+    };
+
+    if (settings.extraParams != null) {
+      extraParams.addAll(settings.extraParams!);
+    }
+
+    createActor(
+      className: lightCardClassName,
+      settings: settings.copyWith(
+        name: settings.name ?? _localizations.newActorNameFlag,
+        template: UnrealTemplateData(name: '', path: templatePath),
+        extraParams: extraParams,
+      ),
+    );
+
+    _recentActorSettings.addRecentlyPlacedActor(RecentlyPlacedActorData(
+      type: RecentlyPlacedActorType.flag,
+    ));
+  }
+
+  /// Send a message to the engine requesting that a new CCW be created.
+  void createColorCorrectionWindow([ActorCreationSettings settings = ActorCreationSettings.defaults]) {
+    if (settings.name == null) {
+      settings = settings.copyWith(name: settings.name ?? _localizations.newActorNameColorCorrectWindow);
+    }
+
+    createActor(
+      className: colorCorrectWindowClassNames.first,
+      settings: settings.copyWith(
+        mapPosition: settings.mapPosition ?? Offset(0.5, 0.5),
+      ),
+    );
+
+    _recentActorSettings.addRecentlyPlacedActor(RecentlyPlacedActorData(
+      type: RecentlyPlacedActorType.colorCorrectWindow,
+    ));
+  }
+
+  /// Send a message to the engine requesting that a new actor be created.
+  void createActor({
+    required String className,
+    ActorCreationSettings settings = ActorCreationSettings.defaults,
+  }) {
+    Offset? position = settings.mapPosition;
+
+    if (position == null && (settings.template == null || settings.template?.path == 'None')) {
+      // Place at camera center
+      position = _previewRenderManager.rendererFocalPoint;
+    }
+
+    final messageBody = {
+      'RendererId': _previewRenderManager.rendererId,
+      if (settings.name != null) 'ActorName': settings.name,
+      if (settings.template != null) 'TemplatePath': settings.template?.path,
+      'ActorClass': className,
+      'RequestId': _startActorCreationRequest(callback: settings.callback, position: position),
+      'OverridePosition': position != null,
+      if (position != null) 'Position': offsetToJson(position),
+    };
+
+    if (settings.extraParams != null) {
+      messageBody.addAll(settings.extraParams!);
+    }
+
+    _connectionManager.sendMessage('ndisplay.preview.actor.create', messageBody);
+  }
+
+  /// Send a message to the engine requesting that the currently selected actors be duplicated.
+  void duplicateSelectedActors({void Function()? callback}) {
+    final List<String> actorPathsToDuplicate = _selectedActorSettings.selectedActors
+        .getValue()
+        .where((String actorPath) => _actorManager.getActorAtPath(actorPath)?.bIsValid == true)
+        .toList(growable: false);
+
+    _connectionManager.sendMessage('stageapp.actors.duplicate', {
+      'Actors': actorPathsToDuplicate,
+      'RequestId': _startActorCreationRequest(callback: callback),
+    });
+  }
+
+  /// Create an ID for an actor creation request and start a timer for it.
+  /// When the engine confirms that the actor has been created, [callback] will be called.
+  /// If [position] is provided and the API version requires it, we will attempt to move the actor to that position
+  /// after its creation.
+  int _startActorCreationRequest({void Function()? callback, Offset? position}) {
+    final int requestId = _lastRequestId;
+    ++_lastRequestId;
+
+    _pendingRequests[requestId] = _PendingCreationRequest(callback: callback, position: position);
+    return requestId;
+  }
+
+  /// Called when the engine sends a message that actors we requested were created.
+  void _onRequestedActorsCreated(dynamic message) {
+    final dynamic requestId = message['RequestId'];
+    final _PendingCreationRequest? pendingRequest = _pendingRequests.remove(requestId);
+    if (pendingRequest == null) {
+      return;
+    }
+
+    final List<dynamic>? actorPaths = message['ActorPaths'];
+    if (pendingRequest.bSelectOnComplete) {
+      if (actorPaths != null) {
+        // We got a list of actors, so drop the current selection and select those ones
+        _selectedActorSettings.deselectAllActors();
+
+        for (dynamic actorPath in actorPaths) {
+          if (actorPath is! String) {
+            continue;
+          }
+          _selectedActorSettings.selectActor(actorPath);
+        }
+      }
+    }
+
+    // On older versions of the engine, a bug causes the initial actor position to be incorrect in some cases, so send
+    // a quick drag request to fix it if possible.
+    if (_connectionManager.apiVersion?.bIsNewActorOverridePositionAccurate != true &&
+        pendingRequest.position != null &&
+        actorPaths != null &&
+        actorPaths.length == 1 &&
+        actorPaths.first is String) {
+      if (_previewRenderManager.bIsDraggingActors) {
+        _log.info('Could not fix position for new actor ${actorPaths[0]} because a drag was in progress');
+      } else {
+        _previewRenderManager.beginActorDrag(actors: [actorPaths.first] /*, primaryActor: actorPaths.first*/);
+        _previewRenderManager.endActorDrag(pendingRequest.position!);
+      }
+    }
+
+    pendingRequest.callback?.call();
+    pendingRequest.dispose();
+  }
+}
+
+/// Data about a request to create actors pending a response from the engine.
+class _PendingCreationRequest {
+  _PendingCreationRequest({this.callback, this.position}) {
+    _autoSelectTimeoutTimer = Timer(_creationRequestTimeout, _onAutoSelectTimedOut);
+  }
+
+  /// Callback function for when the request completes.
+  final void Function()? callback;
+
+  /// The actor's desired position.
+  final Offset? position;
+
+  /// Timer that fires when the request took too long to automatically select the new actors.
+  late final Timer _autoSelectTimeoutTimer;
+
+  /// Whether to automatically select the new actors when the request completes.
+  bool bSelectOnComplete = true;
+
+  /// Clean up any remaining data.
+  void dispose() {
+    _autoSelectTimeoutTimer.cancel();
+  }
+
+  /// Called when a creation request has waited too long for a response to select the actors automatically.
+  void _onAutoSelectTimedOut() {
+    bSelectOnComplete = false;
+  }
+}
