Branch snapshot for CL 27138787

diff --git a/Engine/Extras/VirtualProduction/EpicStageApp/lib/models/unreal_actor_manager.dart b/Engine/Extras/VirtualProduction/EpicStageApp/lib/models/unreal_actor_manager.dart
new file mode 100644
--- /dev/null
+++ b/Engine/Extras/VirtualProduction/EpicStageApp/lib/models/unreal_actor_manager.dart
@@ -0,0 +1,585 @@
+// Copyright Epic Games, Inc. All Rights Reserved.
+
+import 'dart:async';
+import 'dart:collection';
+
+import 'package:flutter/foundation.dart';
+import 'package:flutter/material.dart';
+import 'package:logging/logging.dart';
+import 'package:provider/provider.dart';
+
+import '../utilities/unreal_utilities.dart';
+import 'actor_data/generic_actor_data.dart';
+import 'actor_data/light_card_actor_data.dart';
+import 'actor_data/per_class_actor_data.dart';
+import 'engine_connection.dart';
+import 'settings/selected_actor_settings.dart';
+import 'unreal_class_data.dart';
+import 'unreal_property_manager.dart';
+import 'unreal_types.dart';
+
+final _log = Logger('UnrealActorManager');
+
+/// Callback for when the list of actors of Unreal type [className] changes.
+typedef ActorUpdateCallback = Function(ActorUpdateDetails details);
+
+/// Details object for callbacks that use [ActorUpdateCallback].
+class ActorUpdateDetails {
+  const ActorUpdateDetails({
+    required this.className,
+    required this.addedActors,
+    required this.renamedActors,
+    required this.deletedActors,
+    this.bIsDueToDisconnect = false,
+  });
+
+  /// The name of class for which we received an update.
+  final String className;
+
+  /// Actors of this class that were added.
+  final List<UnrealObject> addedActors;
+
+  /// Actors of this class that were renamed.
+  final List<UnrealObject> renamedActors;
+
+  /// Actors of this class that were deleted.
+  final List<UnrealObject> deletedActors;
+
+  /// If true, this update is because the app has been disconnected from the engine.
+  final bool bIsDueToDisconnect;
+}
+
+/// Function that takes per-class actor data parameters and creates a data class.
+typedef _ActorDataCreatorFunction<T extends UnrealPerClassActorData> = T Function(UnrealPerClassActorDataParams);
+
+/// Stores an [_ActorDataCreatorFunction] and allows us to query the actor data type.
+class _ActorDataCreator<T extends UnrealPerClassActorData> {
+  const _ActorDataCreator(this._fn)
+      : assert(T != UnrealPerClassActorData, '_ActorDataCreator must use a concrete actor data type');
+
+  /// The function this will use to create a new data object.
+  final _ActorDataCreatorFunction<T> _fn;
+
+  /// The type of actor data object this will create.
+  Type get dataType => T;
+
+  /// Create a new actor data object of the appropriate type.
+  T create(UnrealPerClassActorDataParams params) => _fn(params);
+}
+
+/// Maintains information about and provides event hooks for actors in the engine that the app cares about.
+class UnrealActorManager with WidgetsBindingObserver {
+  /// Construct an actor manager and register it with the connection manager.
+  UnrealActorManager(BuildContext context)
+      : _connectionManager = Provider.of<EngineConnectionManager>(context, listen: false),
+        _propertyManager = Provider.of<UnrealPropertyManager>(context, listen: false),
+        _selectedActorSettings = Provider.of<SelectedActorSettings>(context, listen: false) {
+    _registerAllActorDataCreators();
+
+    _connectionManager
+      ..registerConnectionStateListener(_onConnectionStateChanged)
+      ..registerMessageListener('ActorsChanged', _onActorsChanged);
+    _onConnectionStateChanged(_connectionManager.connectionState);
+
+    _selectedActorSettings.displayClusterRootPath.listen(_onRootActorChanged);
+
+    WidgetsBinding.instance.addObserver(this);
+  }
+
+  /// The connection manager we'll use to communicate with the engine.
+  final EngineConnectionManager _connectionManager;
+
+  /// The property manager we'll use to retrieve additional actor data.
+  final UnrealPropertyManager _propertyManager;
+
+  /// Settings we need access to for actor-related operations.
+  final SelectedActorSettings _selectedActorSettings;
+
+  /// Map from watched class names to the callback functions to be called when they change.
+  final Map<String, Set<ActorUpdateCallback>> _classWatchCallbacks = {};
+
+  /// Callback functions for when any actor is updated.
+  final Set<ActorUpdateCallback> _anyClassWatchCallbacks = {};
+
+  /// Map from an actor's path in the editor to the data we've stored about it locally.
+  final Map<String, UnrealObject> _actorsByPath = {};
+
+  /// Map from watched class name to the list of actors that have been seen with that class name.
+  final Map<String, Set<UnrealObject>> _actorsByClass = {};
+
+  /// Map from watched class name to completers to complete when we receive the initial list of actor of that class
+  /// from the engine.
+  final Map<String, Completer<Set<UnrealObject>>> _firstUpdateCompleters = {};
+
+  /// Map from class name to a function that creates an [UnrealPerClassActorData] for a given class name.
+  final Map<UnrealClass, List<_ActorDataCreator>> _actorDataCreatorsByClass = {};
+
+  /// Set of light cards that belong to the current DCRA.
+  ValueNotifier<UnmodifiableSetView<UnrealObject>> _currentRootActorLightCards = ValueNotifier(UnmodifiableSetView({}));
+
+  /// Set of light cards that belong to the current DCRA.
+  ValueListenable<UnmodifiableSetView<UnrealObject>> get currentRootActorLightCards => _currentRootActorLightCards;
+
+  /// Start watching for any engine-side changes to actors of the Unreal type [className], and call [callback] when
+  /// there's any change in the list of relevant actors.
+  void watchClassName(String className, ActorUpdateCallback callback) {
+    if (_classWatchCallbacks.containsKey(className)) {
+      // We're already watching this class, so just add the callback
+      _classWatchCallbacks[className]!.add(callback);
+      return;
+    }
+
+    _classWatchCallbacks[className] = {callback};
+
+    if (!_firstUpdateCompleters.containsKey(className)) {
+      _firstUpdateCompleters[className] = Completer<Set<UnrealObject>>();
+    }
+
+    _connectionManager.sendRawMessage(_createActorRegisterMessage(className, true));
+  }
+
+  /// Stop watching for engine-side changes to actors of the Unreal type [className] with the given [callback].
+  void stopWatchingClassName(String className, ActorUpdateCallback callback) {
+    final Set<ActorUpdateCallback>? callbacks = _classWatchCallbacks[className];
+    if (callbacks == null) {
+      return;
+    }
+
+    callbacks.remove(callback);
+
+    if (callbacks.isNotEmpty) {
+      return;
+    }
+
+    // No more watchers of this class
+    _classWatchCallbacks.remove(className);
+    _connectionManager.sendRawMessage(_createActorRegisterMessage(className, false));
+
+    // Forget about actors of this class
+    final Set<UnrealObject> actorsToForget = getActorsOfClass(className);
+    _actorsByClass.remove(className);
+
+    if (actorsToForget.isEmpty) {
+      return;
+    }
+
+    // Filter out any actors we're about to forget that are also part of another watched class
+    for (final Set<UnrealObject> relevantActors in _actorsByClass.values) {
+      actorsToForget.removeWhere((UnrealObject actor) => relevantActors.contains(actor));
+    }
+
+    for (final UnrealObject actor in actorsToForget) {
+      actor.onUnwatched();
+      _actorsByPath.remove(actor.path);
+    }
+  }
+
+  /// Watch for updates to all actor classes which are already watched by previous or future calls to [watchClassName].
+  /// Note that this will not subscribe to any new classes in the engine.
+  void watchExistingSubscriptions(ActorUpdateCallback callback) {
+    _anyClassWatchCallbacks.add(callback);
+  }
+
+  /// Stop watching for updates to all actor classes for a callback that was previously passed to
+  /// [watchExistingSubscriptions].
+  void stopWatchingExistingSubscriptions(ActorUpdateCallback callback) {
+    _anyClassWatchCallbacks.remove(callback);
+  }
+
+  /// Check whether a [className] is currently being watched by at least one watcher.
+  bool isClassWatched(String className) {
+    final callbacks = _classWatchCallbacks[className];
+    return callbacks != null && callbacks.isNotEmpty;
+  }
+
+  /// Get the watched [UnrealObject] corresponding to a path name.
+  UnrealObject? getActorAtPath(String path) {
+    return _actorsByPath[path];
+  }
+
+  /// Get the cached set of watched [UnrealObject]s corresponding to a [className]. Note that if the class is not
+  /// currently watched, this will always return an empty set.
+  Set<UnrealObject> getActorsOfClass(String className) {
+    return _actorsByClass[className]?.toSet() ?? {};
+  }
+
+  /// Get the cached set of [UnrealObject]s corresponding to a [className], but wait until the initial list of actors
+  /// has been received from the engine if it hasn't already.
+  /// If [className] isn't watched, this will automatically watch the name for long enough to retrieve the list, then
+  /// stop watching it.
+  Future<Set<UnrealObject>> getInitialActorsOfClass(String className) {
+    // If nobody is already watching this class, do so temporarily
+    late final Function(ActorUpdateDetails) tempWatcherStub;
+    final bool bNeedsTempWatch = !isClassWatched(className);
+
+    if (bNeedsTempWatch) {
+      tempWatcherStub = (_) {};
+      watchClassName(className, tempWatcherStub);
+    }
+
+    late final Future<Set<UnrealObject>> result;
+
+    final Completer<Set<UnrealObject>>? completer = _firstUpdateCompleters[className];
+
+    if (completer == null) {
+      // If there's no completer, we already have the initial list of actors
+      result = Future.value(getActorsOfClass(className));
+    } else {
+      // Otherwise, wait until the initial list is received
+      result = completer.future;
+    }
+
+    if (bNeedsTempWatch) {
+      result.then((_) => stopWatchingClassName(className, tempWatcherStub));
+    }
+
+    return result;
+  }
+
+  @override
+  void didChangeAppLifecycleState(AppLifecycleState state) {
+    super.didChangeAppLifecycleState(state);
+
+    switch (state) {
+      case AppLifecycleState.resumed:
+        _subscribeToEngineActors(true);
+        break;
+
+      case AppLifecycleState.paused:
+      case AppLifecycleState.detached:
+        _subscribeToEngineActors(false);
+        break;
+
+      default:
+        break;
+    }
+  }
+
+  /// Register [UnrealPerClassActorData]s for relevant class types.
+  void _registerAllActorDataCreators() {
+    _registerActorDataCreatorForClass(UnrealClassRegistry.lightCard, (params) => LightCardActorData(params));
+  }
+
+  /// Register a creator function for a class.
+  /// We use a generic function for this instead of adding directly to the list because otherwise, Dart will infer T to
+  /// be the base UnrealPerClassActorData class instead of the specific type. We need the specific type later so we can
+  /// check if the actor already has the corresponding data attached.
+  void _registerActorDataCreatorForClass<T extends UnrealPerClassActorData>(
+    UnrealClass unrealClass,
+    _ActorDataCreatorFunction<T> creatorFn,
+  ) {
+    List<_ActorDataCreator>? creatorList = _actorDataCreatorsByClass[unrealClass];
+    if (creatorList == null) {
+      creatorList = [];
+      _actorDataCreatorsByClass[unrealClass] = creatorList;
+    }
+
+    creatorList.add(_ActorDataCreator<T>(creatorFn));
+  }
+
+  /// Subscribe to the WebSocket route that gives us updates about actors
+  void _subscribeToEngineActors(bool bShouldRegister) {
+    final List<dynamic> messages = [];
+
+    // Re-register for any actor classes that we're still watching
+    for (String className in _classWatchCallbacks.keys) {
+      messages.add(_createActorRegisterMessage(className, true));
+    }
+
+    _connectionManager.sendBatchedMessage(messages);
+  }
+
+  /// Called when we lose connection to the engine.
+  void _onEngineDisconnected() {
+    final Set<UnrealObject> deletedActors = {};
+
+    // Send destroy event for all actors since we will no longer know if they exist
+    for (UnrealObject actor in _actorsByPath.values) {
+      deletedActors.add(actor);
+      actor.onDestroyed();
+    }
+
+    // Send destroy event for each class
+    for (MapEntry<String, Set<UnrealObject>> classPair in _actorsByClass.entries) {
+      final className = classPair.key;
+
+      final Set<ActorUpdateCallback>? callbacks = _classWatchCallbacks[classPair.key];
+      if (callbacks == null) {
+        continue;
+      }
+
+      final updateDetails = ActorUpdateDetails(
+        className: className,
+        addedActors: [],
+        renamedActors: [],
+        deletedActors: classPair.value.toList(growable: false),
+        bIsDueToDisconnect: true,
+      );
+
+      for (ActorUpdateCallback callback in callbacks) {
+        callback(updateDetails);
+      }
+
+      for (ActorUpdateCallback callback in _anyClassWatchCallbacks) {
+        callback(updateDetails);
+      }
+    }
+
+    // Dispose actors now that nobody should be accessing them
+    for (UnrealObject actor in deletedActors) {
+      actor.dispose();
+    }
+
+    _actorsByClass.clear();
+    _actorsByPath.clear();
+    _currentRootActorLightCards.value = UnmodifiableSetView({});
+  }
+
+  /// Called when the connection manager's connection state changes.
+  void _onConnectionStateChanged(EngineConnectionState connectionState) {
+    switch (connectionState) {
+      case EngineConnectionState.connected:
+        _subscribeToEngineActors(true);
+        break;
+
+      case EngineConnectionState.disconnected:
+        _onEngineDisconnected();
+        break;
+
+      default:
+        break;
+    }
+  }
+
+  /// Called when the actors in the engine have changed.
+  void _onActorsChanged(dynamic message) {
+    var changes = message['Changes'];
+    if (changes == null) {
+      return;
+    }
+
+    changes.forEach((String className, dynamic classChanges) {
+      if (_classWatchCallbacks.keys.contains(className)) {
+        _updateActorsFromEngine(className, classChanges);
+      }
+
+      final Completer<Set<UnrealObject>>? completer = _firstUpdateCompleters[className];
+      if (completer != null) {
+        completer.complete(getActorsOfClass(className));
+        _firstUpdateCompleters.remove(className);
+      }
+    });
+
+    _updateRootActorLightCards();
+  }
+
+  /// Create a message to register/unregister a [className] for updates from the engine.
+  /// If [bShouldRegister] is false, unregister for the class name.
+  dynamic _createActorRegisterMessage(String className, bool bShouldRegister) {
+    return createUnrealWebSocketMessage(
+      bShouldRegister ? 'actors.register' : 'actors.unregister',
+      {'ClassName': className},
+    );
+  }
+
+  /// Update our actor data about actors of class [className] based on change events from the engine.
+  void _updateActorsFromEngine(String className, dynamic classChanges) {
+    // Remove any actors that no longer exist. Do this first in case the actor was also re-added this frame
+    final List? deletedActorPaths = classChanges['DeletedActors'];
+    final List<UnrealObject> deletedActors = [];
+
+    if (deletedActorPaths != null) {
+      for (var actorData in deletedActorPaths) {
+        final String? path = actorData['Path'];
+        if (path == null) {
+          continue;
+        }
+
+        final UnrealObject? actor = _actorsByPath[path];
+        if (actor == null) {
+          continue;
+        }
+
+        deletedActors.add(actor);
+        actor.onDestroyed();
+        _actorsByPath.remove(path);
+
+        for (final Set<UnrealObject> classActors in _actorsByClass.values) {
+          classActors.remove(actor);
+        }
+      }
+    }
+
+    // Add any new actors
+    final List? addedActorPaths = classChanges['AddedActors'];
+    final List<UnrealObject> addedActors = [];
+
+    if (addedActorPaths != null) {
+      for (var actorData in addedActorPaths) {
+        final String? path = actorData['Path'];
+        final String? name = actorData['Name'];
+
+        if (path == null || name == null) {
+          continue;
+        }
+
+        // Create a new entry if this is a new actor
+        UnrealObject? actor = _actorsByPath[path];
+        if (actor == null) {
+          actor = UnrealObject(name: name, path: path);
+          _actorsByPath[path] = actor;
+        } else {
+          // Update the actor's name if a new one has been created at the same path
+          actor.name = name;
+        }
+
+        // Add this actor to the class list. An actor may have multiple watched classes, so we want to do this
+        // regardless of whether the actor is actually new.
+        Set<UnrealObject>? classActors = _actorsByClass[className];
+        if (classActors == null) {
+          classActors = {};
+          _actorsByClass[className] = classActors;
+        }
+
+        if (!actor.classNames.contains(className)) {
+          actor.addClassName(className);
+
+          // Add additional data depending on the actor's class
+          final actorDataParams = UnrealPerClassActorDataParams(
+            actor: actor,
+            propertyManager: _propertyManager,
+            connectionManager: _connectionManager,
+          );
+
+          // All actors should have generic data
+          if (actor.getPerClassData<GenericActorData>() == null) {
+            actor.addPerClassData(GenericActorData(actorDataParams));
+          }
+
+          // Add any specific class data that wasn't already present
+          UnrealClass? unrealClass = actor.unrealClass;
+
+          while (unrealClass != null) {
+            final List<_ActorDataCreator>? actorDataCreators = _actorDataCreatorsByClass[unrealClass];
+            if (actorDataCreators != null) {
+              for (final _ActorDataCreator creator in actorDataCreators) {
+                if (actor.getPerClassDataOfType(creator.dataType) != null) {
+                  continue;
+                }
+
+                actor.addPerClassData(creator.create(actorDataParams));
+              }
+            }
+
+            unrealClass = unrealClass.parent;
+          }
+        }
+
+        classActors.add(actor);
+        addedActors.add(actor);
+      }
+    }
+
+    // Rename existing actors
+    final List renamedActorPaths = classChanges['RenamedActors'] ?? [];
+    final List<UnrealObject> renamedActors = [];
+
+    for (var actorData in renamedActorPaths) {
+      final String? path = actorData['Path'];
+      final String? name = actorData['Name'];
+
+      if (path == null || name == null) {
+        continue;
+      }
+
+      final UnrealObject? actor = _actorsByPath[path];
+      if (actor == null) {
+        continue;
+      }
+
+      actor.name = name;
+      renamedActors.add(actor);
+    }
+
+    final updateDetails = ActorUpdateDetails(
+      className: className,
+      addedActors: addedActors,
+      renamedActors: renamedActors,
+      deletedActors: deletedActors,
+    );
+
+    // Call callbacks for watchers of this class
+    final Set<ActorUpdateCallback>? callbacks = _classWatchCallbacks[className];
+    if (callbacks != null) {
+      for (final ActorUpdateCallback callback in callbacks) {
+        callback(updateDetails);
+      }
+    }
+
+    // Call callbacks for watchers of any class
+    for (ActorUpdateCallback callback in _anyClassWatchCallbacks) {
+      callback(updateDetails);
+    }
+
+    // Dispose of actors that should no longer be referenced
+    for (UnrealObject actor in deletedActors) {
+      actor.dispose();
+    }
+  }
+
+  /// Send a request to update the list of light cards for the current DCRA.
+  void _updateRootActorLightCards() async {
+    if (_connectionManager.connectionState != EngineConnectionState.connected) {
+      return;
+    }
+
+    final String rootPath = _selectedActorSettings.displayClusterRootPath.getValue();
+    if (rootPath.isEmpty) {
+      return;
+    }
+
+    final UnrealHttpResponse response = await _connectionManager.sendHttpRequest(UnrealHttpRequest(
+      url: '/remote/object/call',
+      verb: 'PUT',
+      body: {
+        'objectPath': '/Script/DisplayCluster.Default__DisplayClusterBlueprintLib',
+        'functionName': 'FindLightCardsForRootActor',
+        'parameters': {
+          'rootActor': rootPath,
+        },
+      },
+    ));
+
+    _currentRootActorLightCards.value = UnmodifiableSetView({});
+
+    if (response.code != 200) {
+      _log.warning('Tried to get light cards for root actor ("$rootPath"), but got a ${response.code} response');
+      return;
+    }
+
+    final List? lightCardPaths = response.body['OutLightCards'];
+    if (lightCardPaths == null) {
+      _log.warning('Tried to get light cards for root actor ("$rootPath"), but body contained no OutLightCards');
+      return;
+    }
+
+    final Set<UnrealObject> lightCards = {};
+    for (final String path in lightCardPaths) {
+      final UnrealObject? lightCard = getActorAtPath(path);
+      if (lightCard == null) {
+        continue;
+      }
+
+      lightCards.add(lightCard);
+    }
+
+    _currentRootActorLightCards.value = UnmodifiableSetView(lightCards);
+  }
+
+  /// Called when the selected DCRA changes.
+  void _onRootActorChanged(_) {
+    // Clear the set first so we immediately hide irrelevant actors
+    _currentRootActorLightCards.value = UnmodifiableSetView({});
+    _updateRootActorLightCards();
+  }
+}
diff --git a/Engine/Extras/VirtualProduction/EpicStageApp/lib/models/unreal_class_data.dart b/Engine/Extras/VirtualProduction/EpicStageApp/lib/models/unreal_class_data.dart
new file mode 100644
--- /dev/null
+++ b/Engine/Extras/VirtualProduction/EpicStageApp/lib/models/unreal_class_data.dart
@@ -0,0 +1,251 @@
+// Copyright Epic Games, Inc. All Rights Reserved.
+
+import 'dart:collection';
+
+import 'package:flutter/material.dart';
+import 'package:flutter_gen/gen_l10n/app_localizations.dart';
+import 'package:logging/logging.dart';
+
+import '../utilities/constants.dart';
+import 'actor_data/light_card_actor_data.dart';
+import 'unreal_types.dart';
+
+final _log = Logger('UnrealClassData');
+
+/// Data about a class that exists in Unreal Engine.
+class UnrealClass {
+  UnrealClass(
+    this.internalName, {
+    required this.getDisplayName,
+    UnrealClass? parent,
+    String? iconPath,
+    Color? color,
+    this.iconPathFunction,
+    this.uiColorFunction,
+    required List<String> classNames,
+  })  : _parent = parent,
+        _iconPath = iconPath,
+        _uiColor = color,
+        classNames = UnmodifiableListView(classNames) {
+    if (parent != null) {
+      parent._children.add(this);
+    }
+  }
+
+  /// Function that returns a user-friendly display name for this class.
+  final String Function(BuildContext) getDisplayName;
+
+  /// Internal display name for this class.
+  final String internalName;
+
+  /// Path to the icon representing this class, or null to defer to parent class.
+  final String? _iconPath;
+
+  /// The color associated with this class, or null to defer to the parent class.
+  final Color? _uiColor;
+
+  /// Optional function to determine the icon based on an instance of this class.
+  final String? Function(UnrealObject object)? iconPathFunction;
+
+  /// Optional function to determine the color based on an instance of this class.
+  final Color? Function(UnrealObject object)? uiColorFunction;
+
+  /// Class names that map to this class.
+  final UnmodifiableListView<String> classNames;
+
+  /// The class this descends from.
+  final UnrealClass? _parent;
+
+  /// Descendants of this class.
+  final List<UnrealClass> _children = [];
+
+  /// The default icon path representing this class.
+  String? get defaultIconPath => _iconPath ?? _parent?._iconPath;
+
+  /// The default color representing this class.
+  Color? get defaultColor => _uiColor ?? _parent?._uiColor;
+
+  /// The class this descends from.
+  UnrealClass? get parent => _parent;
+
+  /// Descendants of this class.
+  UnmodifiableListView<UnrealClass> get children => UnmodifiableListView(_children);
+
+  /// Get the icon path for an instance of this class.
+  String? getIconPathForInstance(UnrealObject object) {
+    if (iconPathFunction == null || object.unrealClass?.isA(this) != true) {
+      return defaultIconPath;
+    }
+
+    return iconPathFunction!(object) ?? defaultIconPath;
+  }
+
+  /// Get the associated color for an instance of this class.
+  Color? getUIColorForInstance(UnrealObject object) {
+    if (uiColorFunction == null || object.unrealClass?.isA(this) != true) {
+      return defaultColor;
+    }
+
+    return uiColorFunction!(object) ?? defaultColor;
+  }
+
+  /// Check if this class is or descends from another class.
+  bool isA(UnrealClass otherClass) {
+    if (otherClass == this) {
+      return true;
+    }
+
+    if (parent != null) {
+      return parent!.isA(otherClass);
+    }
+
+    return false;
+  }
+
+  @override
+  String toString() {
+    return internalName;
+  }
+}
+
+/// Holds data about the hierarchy of all known Unreal Engine classes.
+class UnrealClassRegistry {
+  static final actor = UnrealClass(
+    'Actor',
+    getDisplayName: (context) => AppLocalizations.of(context)!.actorNameActor,
+    classNames: ['/Script/Engine.Actor'],
+  );
+
+  static final nDisplayRootActor = UnrealClass(
+    'DisplayClusterRootActor',
+    getDisplayName: (context) => AppLocalizations.of(context)!.actorNameDisplayClusterRoot,
+    iconPath: 'assets/images/icons/ndisplay.svg',
+    parent: actor,
+    classNames: [nDisplayRootActorClassName],
+  );
+
+  static final lightCard = UnrealClass(
+    'LightCard',
+    getDisplayName: (context) => AppLocalizations.of(context)!.actorNameLightCard,
+    iconPath: 'assets/images/icons/light_card.svg',
+    color: const Color(0xffc4a300),
+    iconPathFunction: (UnrealObject object) {
+      final lightCardData = object.getPerClassData<LightCardActorData>();
+      if (lightCardData != null) {
+        if (!lightCardData.bIsDataReady) {
+          return null;
+        }
+
+        if (lightCardData.bIsUV == true) {
+          return 'assets/images/icons/light_card_uv.svg';
+        }
+
+        if (lightCardData.bIsFlag == true) {
+          return 'assets/images/icons/light_card_flag.svg';
+        }
+      }
+
+      return null;
+    },
+    uiColorFunction: (UnrealObject object) {
+      final lightCardData = object.getPerClassData<LightCardActorData>();
+      if (lightCardData != null) {
+        if (!lightCardData.bIsDataReady) {
+          return null;
+        }
+
+        if (lightCardData.bIsFlag == true) {
+          return const Color(0xff6043d2);
+        }
+      }
+
+      return null;
+    },
+    parent: actor,
+    classNames: [lightCardClassName],
+  );
+
+  static final chromakeyCard = UnrealClass(
+    'ChromakeyCard',
+    getDisplayName: (context) => AppLocalizations.of(context)!.actorNameChromakeyCard,
+    iconPath: 'assets/images/icons/chromakey_card.svg',
+    color: const Color(0xff16884a),
+    parent: lightCard,
+    classNames: [chromakeyCardClassName],
+  );
+
+  static final colorCorrectRegion = UnrealClass(
+    'ColorCorrectRegion',
+    getDisplayName: (context) => AppLocalizations.of(context)!.actorNameColorCorrectRegion,
+    color: const Color(0xffd72e80),
+    iconPath: 'assets/images/icons/color_correct_region.svg',
+    parent: actor,
+    classNames: [colorCorrectRegionClassName],
+  );
+
+  static final colorCorrectWindow = UnrealClass(
+    'ColorCorrectWindow',
+    getDisplayName: (context) => AppLocalizations.of(context)!.actorNameColorCorrectWindow,
+    iconPath: 'assets/images/icons/color_correct_window.svg',
+    parent: colorCorrectRegion,
+    classNames: colorCorrectWindowClassNames,
+  );
+
+  static final postProcessVolume = UnrealClass(
+    'PostProcessVolume',
+    getDisplayName: (context) => AppLocalizations.of(context)!.actorNamePostProcessVolume,
+    iconPath: 'assets/images/icons/post_process_volume.svg',
+    parent: actor,
+    classNames: [postProcessVolumeClassName],
+  );
+
+  /// Determine the lowest-level Unreal class to which an Unreal [object] belongs based on its list of class names.
+  static UnrealClass? getClassForObject(UnrealObject object) {
+    UnrealClass? bestClass;
+
+    for (final String className in object.classNames) {
+      final UnrealClass? newClass = _instance._classNameToClassData[className];
+
+      if (newClass == null) {
+        continue;
+      }
+
+      if (bestClass == null || newClass.isA(bestClass)) {
+        // Either we had no best class, or this class is a subset of the previous best class
+        bestClass = newClass;
+      } else if (!bestClass.isA(newClass)) {
+        // The previous best class is unrelated to the latest class, so warn and ignore it
+        _log.warning('Object $object has conflicting classes ($bestClass and $newClass)');
+      }
+    }
+
+    return bestClass;
+  }
+
+  static final _instance = UnrealClassRegistry._();
+
+  /// Map from class names to our data about that class.
+  final Map<String, UnrealClass> _classNameToClassData = {};
+
+  UnrealClassRegistry._() {
+    _registerClass(actor);
+    _registerClass(nDisplayRootActor);
+    _registerClass(lightCard);
+    _registerClass(chromakeyCard);
+    _registerClass(colorCorrectRegion);
+    _registerClass(colorCorrectWindow);
+    _registerClass(postProcessVolume);
+  }
+
+  /// Register a class for lookup in the registry.
+  void _registerClass(UnrealClass newClass) {
+    for (final String className in newClass.classNames) {
+      if (_classNameToClassData.containsKey(className)) {
+        throw Exception('Tried to register a class with Unreal class name "$className", but another class with this '
+            'path already exists (${_classNameToClassData[className]})');
+      }
+
+      _classNameToClassData[className] = newClass;
+    }
+  }
+}
diff --git a/Engine/Extras/VirtualProduction/EpicStageApp/lib/models/unreal_engine_beacon.dart b/Engine/Extras/VirtualProduction/EpicStageApp/lib/models/unreal_engine_beacon.dart
new file mode 100644
--- /dev/null
+++ b/Engine/Extras/VirtualProduction/EpicStageApp/lib/models/unreal_engine_beacon.dart
@@ -0,0 +1,468 @@
+// Copyright Epic Games, Inc. All Rights Reserved.
+
+import 'dart:async';
+import 'dart:collection';
+import 'dart:io';
+import 'dart:math';
+import 'dart:typed_data';
+
+import 'package:flutter/cupertino.dart';
+import 'package:flutter/foundation.dart';
+import 'package:logging/logging.dart';
+
+import '../utilities/net_utilities.dart';
+
+final _log = Logger('UnrealEngineBeacon');
+
+/// The address to multicast beacon messages to.
+final InternetAddress _multicastAddress = InternetAddress('230.0.0.2');
+
+/// The port to multicast beacon messages to.
+const int _multicastPort = 6667;
+
+/// The base rate at which to send beacon messages.
+const Duration _beaconInterval = Duration(seconds: 1);
+
+/// How long to wait before considering a connection stale and pruning it from the list.
+const Duration _connectionStaleTime = Duration(seconds: 3);
+
+/// The amount of time before reporting to the user that the beacon has failed to connect.
+const Duration _beaconFailureTimeout = Duration(seconds: 5);
+
+/// The amount to multiply [_beaconInterval] by for each consecutive failure to send a message.
+const double _beaconFailureBackoff = 2.0;
+
+/// Sends UDP multicasts on all available interfaces to search for Unreal Engine instances and maintains a list of
+/// recently seen instances.
+class UnrealEngineBeacon {
+  UnrealEngineBeacon({
+    required BuildContext context,
+    this.onEngineInstancesChanged,
+    this.onBeaconFailure,
+  }) {
+    _updateBeaconSockets().then((value) => _startUpdateTimer());
+  }
+
+  /// Timer used to fire the periodic update.
+  Timer? _updateTimer;
+
+  /// If true, we already reported that all beacons are failing.
+  bool _bHasReportedFailure = false;
+
+  /// Map from addresses for each network interface to our beacon socket for that interface.
+  final Map<InternetAddress, _BeaconSocket> _sockets = {};
+
+  /// List of possible engine connections.
+  final List<ConnectionData> _connections = [];
+
+  /// Callback for when the list of available engine connections changes.
+  final Function()? onEngineInstancesChanged;
+
+  /// Callback for when all beacon sockets are failing, so engines can't be detected.
+  final Function()? onBeaconFailure;
+
+  /// The current list of active engine connections.
+  UnmodifiableListView<ConnectionData> get connections => UnmodifiableListView(_connections);
+
+  /// Clean up all sockets and stop listening for engine instances.
+  void dispose() {
+    _stopUpdateTimer();
+
+    for (final _BeaconSocket socket in _sockets.values) {
+      socket.dispose();
+    }
+    _sockets.clear();
+  }
+
+  /// Pause sending beacon messages if paused.
+  void pause() {
+    _stopUpdateTimer();
+
+    for (final _BeaconSocket socket in _sockets.values) {
+      socket.bIsPaused = true;
+    }
+  }
+
+  /// Resume sending beacon messages if paused.
+  void resume() {
+    _startUpdateTimer();
+
+    for (final _BeaconSocket socket in _sockets.values) {
+      socket.bIsPaused = false;
+    }
+  }
+
+  /// Start/restart the update timer.
+  void _startUpdateTimer() {
+    _stopUpdateTimer();
+
+    _updateTimer = Timer.periodic(_beaconInterval, _onUpdatePeriod);
+  }
+
+  /// Stop the update timer.
+  void _stopUpdateTimer() {
+    if (_updateTimer != null) {
+      _updateTimer!.cancel();
+      _updateTimer = null;
+    }
+  }
+
+  /// Called on each update period.
+  void _onUpdatePeriod([Timer? timer]) {
+    _checkForBeaconSocketFailure();
+    _pruneConnections();
+    _updateBeaconSockets();
+  }
+
+  /// Check if all beacon sockets are failing, and if so, report it.
+  void _checkForBeaconSocketFailure() {
+    bool bAllFailing = true;
+    for (final _BeaconSocket socket in _sockets.values) {
+      if (!socket.bIsFailing) {
+        bAllFailing = false;
+      }
+    }
+
+    if (bAllFailing) {
+      if (!_bHasReportedFailure) {
+        _bHasReportedFailure = true;
+        onBeaconFailure?.call();
+      }
+    } else {
+      _bHasReportedFailure = false;
+    }
+  }
+
+  /// Bind/unbind beacon sockets in order to match the current list of network interfaces.
+  Future<void> _updateBeaconSockets() async {
+    final Map<InternetAddress, String> adapterNames = {};
+    final Set<InternetAddress> addresses = {};
+
+    if (NetworkInterface.listSupported) {
+      // Get the list of all available network interfaces
+      final List<NetworkInterface> interfaces = await NetworkInterface.list();
+
+      for (final NetworkInterface interface in interfaces) {
+        addresses.addAll(interface.addresses);
+
+        for (final InternetAddress address in interface.addresses) {
+          adapterNames[address] = interface.name;
+        }
+      }
+    } else {
+      // We can't get a list of interfaces, so fall back to whatever the OS gives us
+      addresses.add(InternetAddress.anyIPv4);
+      addresses.add(InternetAddress.anyIPv6);
+    }
+
+    // Close beacons we no longer need
+    final List<InternetAddress> staleAddresses = [];
+    for (final InternetAddress existingAddress in _sockets.keys) {
+      if (!addresses.contains(existingAddress)) {
+        staleAddresses.add(existingAddress);
+      }
+    }
+
+    for (final InternetAddress staleAddress in staleAddresses) {
+      _sockets[staleAddress]?.dispose();
+      _sockets.remove(staleAddress);
+    }
+
+    // Create beacons for new interfaces
+    for (final InternetAddress newAddress in addresses) {
+      if (_sockets.containsKey(newAddress)) {
+        continue;
+      }
+
+      final beaconSocket = _BeaconSocket(
+        adapterName: adapterNames[newAddress] ?? 'unknown',
+        address: newAddress,
+        onNewConnectionData: _updateConnection,
+      );
+
+      _sockets[newAddress] = beaconSocket;
+    }
+  }
+
+  /// Update the data for a new or existing connection.
+  void _updateConnection(ConnectionData newConnectionData) {
+    bool bAlreadyExisted = false;
+    bool bAnyChanges = false;
+
+    // Check if we've already seen this connection
+    for (int i = 0; i < _connections.length; ++i) {
+      final ConnectionData oldConnectionData = _connections[i];
+      if (oldConnectionData.uuid == newConnectionData.uuid) {
+        // Check if there were any changes other than last seen time
+        if (oldConnectionData.name != newConnectionData.name ||
+            oldConnectionData.websocketPort != newConnectionData.websocketPort ||
+            oldConnectionData.websocketAddress != newConnectionData.websocketAddress) {
+          bAnyChanges = true;
+        }
+
+        // Refresh the connection data (including last seen time)
+        _connections[i] = newConnectionData;
+
+        bAlreadyExisted = true;
+        break;
+      }
+    }
+
+    if (!bAlreadyExisted) {
+      bAnyChanges = true;
+      _connections.add(newConnectionData);
+    }
+
+    if (bAnyChanges) {
+      _connections.sort((a, b) => a.name.compareTo(b.name));
+      onEngineInstancesChanged?.call();
+    }
+  }
+
+  /// Remove any stale connections from the list
+  void _pruneConnections() {
+    final DateTime now = DateTime.now();
+    final List<int> toRemove = [];
+
+    for (int i = 0; i < _connections.length; ++i) {
+      if (now.difference(_connections[i].lastSeenTime) > _connectionStaleTime) {
+        toRemove.add(i);
+      }
+    }
+
+    for (final int i in toRemove.reversed) {
+      _connections.removeAt(i);
+    }
+
+    if (toRemove.length > 0) {
+      onEngineInstancesChanged?.call();
+    }
+  }
+}
+
+/// A socket open on a specific interface to send and receive beacon messages.
+class _BeaconSocket {
+  _BeaconSocket({
+    required this.adapterName,
+    required this.address,
+    required this.onNewConnectionData,
+  }) {
+    _bindSocket().then((final bool bIsValid) {
+      if (!bIsValid) {
+        return;
+      }
+
+      /// Send the first message immediately so we don't have to wait for an update
+      _sendBeaconMessage();
+    });
+  }
+
+  /// The name of the adapter this is using.
+  final String adapterName;
+
+  /// The address of the interface to send/listen on.
+  final InternetAddress address;
+
+  /// Callback function when new engine connection data has been received.
+  final Function(ConnectionData) onNewConnectionData;
+
+  /// Timer for sending out beacon messages.
+  Timer? _messageTimer;
+
+  /// Timer for when we consider this beacon to have failed.
+  Timer? _failTimer;
+
+  /// Whether the beacon has recently successfully sent a multicast message.
+  bool _bHasRecentlySucceeded = false;
+
+  /// Whether this has been disposed.
+  bool _bWasDisposed = false;
+
+  /// Whether this is temporarily not sending beacon messages.
+  bool _bIsPaused = false;
+
+  /// Whether this encountered an instant failure code and should not attempt to reconnect.
+  bool _bWasInstantFailed = false;
+
+  /// The socket used to send and receive beacon messages.
+  RawDatagramSocket? _socket;
+
+  /// The subscription by which this listens to [_socket].
+  StreamSubscription? _subscription;
+
+  /// The number of times this has consecutively failed to send a message via its socket.
+  int _consecutiveFailures = 0;
+
+  /// Error codes which, if encountered, should cause the beacon to immediately be treated as failed.
+  static Set<int> get _instantFailErrorCodes {
+    if (Platform.isIOS) {
+      return {
+        // Operation not permitted
+        1,
+      };
+    }
+
+    if (Platform.isAndroid) {
+      return {
+        // Network is unreachable
+        101,
+      };
+    }
+
+    return {};
+  }
+
+  /// Whether this is currently failing to send messages to its interface.
+  bool get bIsFailing => !_bHasRecentlySucceeded;
+
+  /// Whether this is currently prevented from sending beacon messages.
+  bool get bIsPaused => _bIsPaused;
+
+  void set bIsPaused(bool bNewValue) {
+    _bIsPaused = bNewValue;
+    if (_bIsPaused) {
+      _stopUpdateTimer();
+    } else {
+      _startUpdateTimer();
+    }
+  }
+
+  /// Stop sending and receiving for this beacon and close its socket.
+  void dispose() {
+    if (_bWasDisposed) {
+      return;
+    }
+
+    _bWasDisposed = true;
+    _giveUp();
+  }
+
+  /// Update the number of consecutive failures and adjust the timer accordingly.
+  void _updateConsecutiveFailures(int newValue) {
+    if (_consecutiveFailures != newValue) {
+      _consecutiveFailures = newValue;
+      _startUpdateTimer();
+    }
+  }
+
+  /// Start sending periodic beacon messages.
+  void _startUpdateTimer() {
+    if (_messageTimer != null) {
+      return;
+    }
+
+    // Send periodic beacons to discover any new Unreal Engine instances
+    final Duration period = _beaconInterval * pow(_beaconFailureBackoff, _consecutiveFailures);
+    _messageTimer = Timer.periodic(period, (_) => _sendBeaconMessage());
+  }
+
+  /// Stop sending periodic beacon messages.
+  void _stopUpdateTimer() {
+    _messageTimer?.cancel();
+    _messageTimer = null;
+  }
+
+  /// Send a beacon message on this socket.
+  void _sendBeaconMessage() {
+    if (_socket != null) {
+      final Uint8List message = makeBeaconMessage();
+      if (_socket!.send(message, _multicastAddress, _multicastPort) > 0) {
+        // We successfully sent a beacon message, so cancel the failure timer and indicate that we should consider it a
+        // new failure if we fail again later.
+        _bHasRecentlySucceeded = true;
+        _failTimer?.cancel();
+        _updateConsecutiveFailures(0);
+      }
+    }
+  }
+
+  /// Bind the socket used to send and receive beacon messages.
+  Future<bool> _bindSocket() async {
+    _stopUpdateTimer();
+
+    _subscription?.cancel();
+    _socket?.close();
+
+    late final RawDatagramSocket newSocket;
+
+    try {
+      newSocket = await RawDatagramSocket.bind(address, _multicastPort, ttl: 4);
+      _log.info('Beacon socket open on "$adapterName" (${address.address})');
+    } catch (error) {
+      _log.warning('Failed to bind socket for "$adapterName" (${address.address}):\n$error');
+      return Future.value(false);
+    }
+
+    _socket = newSocket;
+    _subscription = newSocket.listen(
+      _receiveMessage,
+      onError: _onSocketError,
+      onDone: _onSocketClosed,
+      cancelOnError: false,
+    );
+
+    if (!_bIsPaused) {
+      _startUpdateTimer();
+    }
+
+    return Future.value(true);
+  }
+
+  /// Called when an error occurs on the socket that sends beacon messages.
+  void _onSocketError(Object error, StackTrace stackTrace) {
+    _log.info('Failed to send beacon message on "$adapterName" (${address.address}):\n$error');
+
+    if (error is SocketException && _instantFailErrorCodes.contains(error.osError?.errorCode)) {
+      _log.info('Error code ${error.osError?.errorCode} is an instant fail; closing beacon');
+      _bWasInstantFailed = true;
+      _giveUp();
+      return;
+    }
+
+    _updateConsecutiveFailures(_consecutiveFailures + 1);
+
+    if (_bHasRecentlySucceeded && _failTimer == null) {
+      _failTimer = Timer(_beaconFailureTimeout, _handleBeaconFailed);
+    }
+  }
+
+  /// Called when the beacon has failed for too long.
+  void _handleBeaconFailed() {
+    _bHasRecentlySucceeded = false;
+    _failTimer = null;
+  }
+
+  /// Receive a message on the beacon socket (presumably a reply to a beacon message).
+  void _receiveMessage(RawSocketEvent event) {
+    final Datagram? datagram = _socket!.receive();
+    if (datagram != null) {
+      ConnectionData? newConnection = getConnectionFromBeaconResponse(datagram);
+
+      if (newConnection != null) {
+        onNewConnectionData(newConnection);
+      }
+    }
+  }
+
+  /// Called when the socket closes.
+  void _onSocketClosed() {
+    if (_bWasDisposed || _bWasInstantFailed) {
+      return;
+    }
+
+    _stopUpdateTimer();
+    _bindSocket().then((_) {
+      if (!_bIsPaused) {
+        _startUpdateTimer();
+      }
+    });
+  }
+
+  /// Close the socket and cancel the subscription if necessary, then report failure.
+  void _giveUp() {
+    _subscription?.cancel();
+    _socket?.close();
+    _socket = null;
+    _bHasRecentlySucceeded = false;
+  }
+}
diff --git a/Engine/Extras/VirtualProduction/EpicStageApp/lib/models/unreal_object_filters.dart b/Engine/Extras/VirtualProduction/EpicStageApp/lib/models/unreal_object_filters.dart
new file mode 100644
--- /dev/null
+++ b/Engine/Extras/VirtualProduction/EpicStageApp/lib/models/unreal_object_filters.dart
@@ -0,0 +1,123 @@
+// Copyright Epic Games, Inc. All Rights Reserved.
+
+import 'dart:collection';
+
+import 'package:flutter/material.dart';
+import 'package:flutter_gen/gen_l10n/app_localizations.dart';
+
+import 'actor_data/light_card_actor_data.dart';
+import 'unreal_class_data.dart';
+import 'unreal_types.dart';
+
+/// Data about a filter that can be applied to [UnrealObject]s.
+abstract class UnrealObjectFilter {
+  const UnrealObjectFilter({
+    required this.internalName,
+    required this.getDisplayName,
+    required this.iconPath,
+  });
+
+  /// Unique name to refer to this in code.
+  final String internalName;
+
+  /// Function that returns the name shown to the user.
+  final String Function(BuildContext) getDisplayName;
+
+  /// Path to the image file for the icon representing this category.
+  final String? iconPath;
+
+  /// Returns true if the [object] passes this filter.
+  bool passes(UnrealObject object);
+}
+
+/// An [UnrealObjectFilter] with a custom filtering function.
+class _CustomUnrealObjectFilter extends UnrealObjectFilter {
+  const _CustomUnrealObjectFilter({
+    required super.internalName,
+    required super.getDisplayName,
+    required super.iconPath,
+    required bool Function(UnrealObject object) passFunction,
+  }) : _passFunction = passFunction;
+
+  /// Function that returns true if the [object] is a member of this category.
+  final bool Function(UnrealObject object) _passFunction;
+
+  @override
+  bool passes(UnrealObject object) => _passFunction(object);
+}
+
+/// An [UnrealObjectFilter] that checks against a specific Unreal class.
+class _ClassBasedUnrealObjectFilter extends UnrealObjectFilter {
+  _ClassBasedUnrealObjectFilter(
+    this.unrealClass, {
+    bool Function(UnrealObject object)? passFunction,
+  })  : _passFunction = passFunction,
+        super(
+          internalName: 'class${unrealClass.internalName}',
+          getDisplayName: unrealClass.getDisplayName,
+          iconPath: unrealClass.defaultIconPath,
+        );
+
+  /// The class to which this will filter. Objects must be an exact member of this class, not a descendant.
+  final UnrealClass unrealClass;
+
+  /// Optional function that returns true if the [object] is a member of this category, checked in addition to the
+  /// filter's class constraint.
+  final bool Function(UnrealObject object)? _passFunction;
+
+  @override
+  bool passes(UnrealObject object) {
+    if (object.unrealClass != unrealClass) {
+      return false;
+    }
+
+    if (_passFunction != null && !_passFunction!(object)) {
+      return false;
+    }
+
+    return true;
+  }
+}
+
+/// Holds data about all possible types of object filter.
+class UnrealObjectFilterRegistry {
+  static final _instance = UnrealObjectFilterRegistry._();
+
+  /// List of all possible object filters.
+  final UnmodifiableListView<UnrealObjectFilter> _filters = UnmodifiableListView([
+    _ClassBasedUnrealObjectFilter(
+      UnrealClassRegistry.lightCard,
+      passFunction: (UnrealObject object) {
+        final lightcardData = object.getPerClassData<LightCardActorData>();
+        return lightcardData?.bIsDataReady == true && lightcardData?.bIsFlag == false;
+      },
+    ),
+    _CustomUnrealObjectFilter(
+      internalName: 'Flag',
+      getDisplayName: (context) => AppLocalizations.of(context)!.actorNameFlag,
+      iconPath: 'assets/images/icons/light_card_flag.svg',
+      passFunction: (UnrealObject object) {
+        final lightcardData = object.getPerClassData<LightCardActorData>();
+        return lightcardData?.bIsFlag == true;
+      },
+    ),
+    _ClassBasedUnrealObjectFilter(UnrealClassRegistry.chromakeyCard),
+    _ClassBasedUnrealObjectFilter(UnrealClassRegistry.colorCorrectWindow),
+    _ClassBasedUnrealObjectFilter(UnrealClassRegistry.colorCorrectRegion),
+  ]);
+
+  /// Map from filter internal name to filter.
+  final Map<String, UnrealObjectFilter> _filtersByName = {};
+
+  UnrealObjectFilterRegistry._() {
+    for (final UnrealObjectFilter filter in _filters) {
+      _filtersByName[filter.internalName] = filter;
+    }
+  }
+
+  /// Get a list of all available filters.
+  static UnmodifiableListView<UnrealObjectFilter> get filters => _instance._filters;
+
+  /// Get the filter associated with the given internal name.
+  static UnrealObjectFilter? getFilterByName(String internalName) => _instance._filtersByName[internalName];
+}
diff --git a/Engine/Extras/VirtualProduction/EpicStageApp/lib/models/unreal_property_manager.dart b/Engine/Extras/VirtualProduction/EpicStageApp/lib/models/unreal_property_manager.dart
new file mode 100644
--- /dev/null
+++ b/Engine/Extras/VirtualProduction/EpicStageApp/lib/models/unreal_property_manager.dart
@@ -0,0 +1,1080 @@
+// Copyright Epic Games, Inc. All Rights Reserved.
+
+import 'dart:async';
+
+import 'package:flutter/cupertino.dart';
+import 'package:logging/logging.dart';
+import 'package:provider/provider.dart';
+
+import './engine_connection.dart';
+import './property_modify_operations.dart';
+import './unreal_types.dart';
+import '../utilities/unreal_utilities.dart';
+import 'unreal_transaction_manager.dart';
+
+final _log = Logger('UnrealPropertyManager');
+
+/// Callback for when a tracked property changes.
+typedef PropertyChangeCallback = Function(dynamic newValue);
+
+/// Manages all in-engine properties that are being controlled by the app, including communication with the engine,
+/// caching of values, and notifying of changes.
+class UnrealPropertyManager {
+  /// Construct a property manager and register it with the connection manager.
+  UnrealPropertyManager(BuildContext context)
+      : _connectionManager = Provider.of<EngineConnectionManager>(context, listen: false),
+        _transactionManager = Provider.of<UnrealTransactionManager>(context, listen: false) {
+    _connectionManager.registerConnectionStateListener(_onConnectionStateChanged);
+    _connectionManager.registerMessageListener('PresetFieldsChanged', _onPresetFieldsChanged);
+    _onConnectionStateChanged(_connectionManager.connectionState);
+  }
+
+  /// Rate at which we send property changes back to the engine.
+  static const Duration _tickRate = Duration(milliseconds: 100);
+
+  /// The connection manager we'll use to communicate with the engine.
+  final EngineConnectionManager _connectionManager;
+
+  /// The transaction manager used to maintain editor undo/redo state on our property changes.
+  final UnrealTransactionManager _transactionManager;
+
+  /// A map from [TrackedPropertyId] to the information we're tracking about the property.
+  final Map<TrackedPropertyId, _TrackedProperty> _trackedProperties = {};
+
+  /// Lookup table from a property's preset label to our internal ID for it.
+  final Map<String, TrackedPropertyId> _propertyIdsByLabel = {};
+
+  /// Editor transactions that had activity this tick.
+  /// If the last entry in this list is active, it's the active transaction in the engine.
+  final List<_EngineTransaction> _currentTickTransactions = [];
+
+  /// Changes to be sent to the engine pending the creation of our transient preset.
+  final List<dynamic> _messagesPendingPreset = [];
+
+  /// IDs of tracked properties for which we need to send an expose request.
+  final Set<TrackedPropertyId> _propertiesPendingExpose = {};
+
+  /// IDs of tracked properties for which we need to send an unexpose request.
+  final Set<TrackedPropertyId> _propertiesPendingUnexpose = {};
+
+  /// Timer used to regularly send changes back to the engine while connected.
+  Timer? _tickTimer;
+
+  /// The name of the transient preset we're using to manage properties, or null if it doesn't exist yet.
+  String? _transientPresetName;
+
+  /// The future which will return when the transient preset has been created, or null if we haven't started creating
+  /// a transient preset yet.
+  Future<String>? _transientPresetNameFuture;
+
+  /// The ID of the last newly-tracked property.
+  TrackedPropertyId _lastTrackedId = const TrackedPropertyId.invalid();
+
+  /// The sequence number to send to the engine with the next batch of changes.
+  int _sequenceNumber = 0;
+
+  /// The ID of the active transaction that this created, or [UnrealTransactionManager.invalidId] if there isn't an
+  /// active transaction.
+  int _activeTransactionId = UnrealTransactionManager.invalidId;
+
+  /// Start tracking updates about [property] between the engine and the app, and call [changeCallback] when
+  /// the engine sends a message about a change.
+  /// Returns a future which will provide the tracking ID of the property if it is successfully registered.
+  /// If the property is already tracked with the given callback, this will just return the existing property ID.
+  /// If provided, [conversionMetadata] will be stored and provided to functions converting to/from JSON data. If the
+  /// property is already tracked, its metadata will be replaced by the new value.
+  TrackedPropertyId trackProperty(UnrealProperty property, PropertyChangeCallback changeCallback,
+      {Map<String, dynamic>? conversionMetadata}) {
+    // Check if this property is already tracked
+    for (MapEntry<TrackedPropertyId, _TrackedProperty> trackedEntry in _trackedProperties.entries) {
+      final _TrackedProperty trackedProperty = trackedEntry.value;
+
+      // We're already tracking this property
+      if (trackedProperty.property == property) {
+        if (conversionMetadata != null) {
+          // Update conversion metadata
+          trackedProperty.conversionMetadata = conversionMetadata;
+        }
+
+        trackedProperty.changeCallbacks.add(changeCallback);
+
+        if (trackedProperty.bIsPendingUnexpose) {
+          if (_propertiesPendingUnexpose.remove(trackedEntry.key)) {
+            // The property was waiting to be unexposed, but we hadn't sent the request yet, so forget about the
+            // unexpose entirely.
+            trackedProperty.bIsPendingUnexpose = false;
+          } else {
+            // We already sent an unexpose request (or will do so when the exposeCompleter completes).
+            // We can't guarantee the order in which the engine will receive our requests, so we need to wait for
+            // unexpose and then send another expose request.
+            // This completer will complete as soon as the new expose request is received.
+            trackedProperty.exposeCompleter ??= Completer<void>();
+          }
+
+          return trackedEntry.key;
+        }
+
+        // Property is already good to go, so just return its existing ID
+        return trackedEntry.key;
+      }
+    }
+
+    // We aren't tracking the property yet, so we need to register with the engine first.
+    final TrackedPropertyId newId = _getNextTrackedPropertyId();
+    _trackedProperties[newId] = _TrackedProperty(
+      property: property,
+      initialCallbacks: [changeCallback],
+      conversionMetadata: conversionMetadata ?? {},
+    );
+
+    _exposeProperty(newId);
+
+    return newId;
+  }
+
+  /// Returns true if the property with the given [propertyId] is exposed, otherwise false.
+  bool isPropertyExposed(TrackedPropertyId propertyId) {
+    final _TrackedProperty? property = _trackedProperties[propertyId];
+    if (property == null) {
+      return false;
+    }
+
+    return property.exposeCompleter == null;
+  }
+
+  /// Returns a future that will complete when the property with the given [propertyId] is exposed (or re-exposed).
+  Future<void> waitForProperty(TrackedPropertyId propertyId) {
+    final _TrackedProperty? property = _trackedProperties[propertyId];
+    if (property == null) {
+      return Future.error('Property with id $propertyId does not exist');
+    }
+
+    // If the property is waiting for expose, return that future. Otherwise, it's already exposed.
+    return property.exposeCompleter?.future ?? Future.value();
+  }
+
+  /// Stop tracking the property with ID [propertyId]. [changeCallback] must be the same callback that was initially
+  /// passed in to [trackProperty].
+  void stopTrackingProperty(TrackedPropertyId propertyId, PropertyChangeCallback changeCallback) async {
+    final _TrackedProperty? property = _trackedProperties[propertyId];
+    if (property == null) {
+      return;
+    }
+
+    if (!property.changeCallbacks.remove(changeCallback)) {
+      _log.warning('Tried to stop tracking ${property.label}, but there was no callback to remove');
+    }
+
+    // If nobody is watching this property anymore, unsubscribe from updates to it
+    if (property.changeCallbacks.isEmpty) {
+      if (property.bIsPendingUnexpose) {
+        // Property is already pending unexpose, so we don't need to send another request
+        return;
+      }
+
+      // As soon as this property is exposed, or immediately if it's already exposed, flag it to be unexposed.
+      property.bIsPendingUnexpose = true;
+
+      queueForUnexpose() {
+        _propertiesPendingUnexpose.add(propertyId);
+      }
+
+      if (property.exposeCompleter == null) {
+        // Immediately queue the property to be unexposed if it's already exposed
+        queueForUnexpose();
+      } else if (_propertiesPendingExpose.remove(propertyId)) {
+        // Property was queued for expose, but we hadn't actually sent the message, so just drop it
+        property.exposeCompleter?.complete();
+        property.exposeCompleter = null;
+        _trackedProperties.remove(propertyId);
+      } else {
+        // Wait until it's exposed, then unexpose it
+        waitForProperty(propertyId).then((_) {
+          if (property.changeCallbacks.isNotEmpty) {
+            // In the time we were waiting for the property to be exposed, somebody started listening for it, so we
+            // don't actually want to unexpose it anymore.
+            property.bIsPendingUnexpose = false;
+            return;
+          }
+
+          queueForUnexpose();
+        });
+      }
+    }
+  }
+
+  /// Get the value of a tracked property, or null if the property isn't tracked/ready.
+  dynamic getTrackedPropertyValue(TrackedPropertyId propertyId) {
+    return _getReadyTrackedProperty(propertyId)?.value;
+  }
+
+  /// Get the metadata for a tracked property, or null if the property isn't tracked/ready.
+  dynamic getTrackedPropertyMetaData(TrackedPropertyId propertyId) {
+    return _getReadyTrackedProperty(propertyId)?.metadata;
+  }
+
+  /// Start a tracked property transaction.
+  /// Returns false if there was already an active transaction.
+  /// Note that this creates a transaction within the [UnrealPropertyManager] only -- the actual transaction will not be
+  /// created until updates are ready to send on the next tick. If you want to immediately create a transaction, use
+  /// [UnrealTransactionManager.beginTransaction] instead.
+  bool beginTransaction(String description) {
+    final _EngineTransaction? activeTransaction = _getActiveTransaction();
+
+    if (activeTransaction != null) {
+      _log.warning('Tried to start transaction, but one was already in progress');
+      return false;
+    }
+
+    _currentTickTransactions.add(_EngineTransaction(description: description));
+
+    return true;
+  }
+
+  /// Mark the current tracked property transaction as complete.
+  /// Returns false if there was no active transaction to complete.
+  /// Note that this should only be used to end a transaction started using [UnrealPropertyManager.beginTransaction].
+  bool endTransaction() {
+    final _EngineTransaction? activeTransaction = _getActiveTransaction();
+
+    if (activeTransaction == null) {
+      return false;
+    }
+
+    activeTransaction.bIsActive = false;
+
+    return true;
+  }
+
+  /// Modify the value of a tracked property.
+  /// The property will be retrieved based on the [propertyId], and its current value will by modified by the amount in
+  /// [deltaValue] using the provided [operation].
+  /// [deltaValue] may be optional depending on the operation.
+  /// [beginTransaction] should be called before this if [isEditorTransactionActive] returns false.
+  /// If [minMaxBehaviour] behaviour is provided, it will change how the value is handled when it passes the engine's
+  /// min/max values for the property.
+  /// If [overrideMin] and/or [overrideMax] are provided, they will be used instead of the min/max values provided by
+  /// the engine when applying the [minMaxBehaviour].
+  /// Returns the new value of the property, or null if the property isn't tracked.
+  dynamic modifyTrackedPropertyValue(
+    TrackedPropertyId propertyId,
+    PropertyModifyOperation operation, {
+    dynamic deltaValue,
+    PropertyMinMaxBehaviour minMaxBehaviour = PropertyMinMaxBehaviour.clamp,
+    dynamic overrideMin = null,
+    dynamic overrideMax = null,
+  }) {
+    final _EngineTransaction? activeTransaction = _getActiveTransaction();
+
+    if (activeTransaction == null) {
+      throw Exception('Called modifyTrackedPropertyValue with no active transaction. '
+          'Call beginEditorTransaction first.');
+    }
+
+    final _TrackedProperty? property = _getReadyTrackedProperty(propertyId);
+    if (property == null) {
+      return null;
+    }
+
+    final newChange = _TrackedPropertyChange(operation: operation, deltaValue: deltaValue);
+
+    // Keep the value in the range provided by engine, if necessary
+    switch (minMaxBehaviour) {
+      case PropertyMinMaxBehaviour.loop:
+        newChange.loop(
+          referenceValue: property.value,
+          min: overrideMin ?? property.metadata.minValue,
+          max: overrideMax ?? property.metadata.maxValue,
+        );
+        break;
+
+      case PropertyMinMaxBehaviour.clamp:
+        // Clamp the change as appropriate, and if the result is that we won't actually change anything, early out.
+        final bool bIsMeaningfulChange = newChange.clamp(
+          referenceValue: property.value,
+          min: overrideMin ?? property.metadata.minValue,
+          max: overrideMax ?? property.metadata.maxValue,
+        );
+
+        if (!bIsMeaningfulChange) {
+          return property.value;
+        }
+        break;
+
+      case PropertyMinMaxBehaviour.ignore:
+        break;
+    }
+
+    // Update the list of changes to the property this frame, combining with existing changes if possible
+    final List<_TrackedPropertyChange> changesThisFrame = activeTransaction.changes.putIfAbsent(propertyId, () => []);
+
+    if (changesThisFrame.isNotEmpty && changesThisFrame.last.canCombine(newChange)) {
+      // We can combine this change with the previous one, so we don't need to add a new one
+      changesThisFrame.last.combine(newChange);
+    } else {
+      changesThisFrame.add(newChange);
+    }
+
+    // Update the predicted value
+    property.value = newChange.apply(property.value);
+    return property.value;
+  }
+
+  /// Get the transaction that is currently active in the engine, if any.
+  _EngineTransaction? _getActiveTransaction() {
+    if (_currentTickTransactions.isNotEmpty) {
+      final _EngineTransaction latestTransaction = _currentTickTransactions.last;
+      if (latestTransaction.bIsActive) {
+        return latestTransaction;
+      }
+    }
+
+    return null;
+  }
+
+  /// Marks an Unreal Engine property with the given [propertyId] to be exposed in the next batch of requests.
+  void _exposeProperty(TrackedPropertyId propertyId) {
+    final _TrackedProperty? property = _trackedProperties[propertyId];
+
+    if (property == null) {
+      assert(false, '_exposeProperty should only be called with existing tracked properties');
+      return;
+    }
+
+    _propertiesPendingExpose.add(propertyId);
+
+    // If the property didn't already have a completer, create one so we know we're waiting for expose
+    property.exposeCompleter ??= Completer();
+  }
+
+  /// Create a request to the engine to expose a property with the given [propertyId] on the transient preset.
+  UnrealHttpRequestWithCallback _createExposePropertyRequest(TrackedPropertyId propertyId) {
+    final _TrackedProperty? property = _trackedProperties[propertyId];
+
+    assert(property != null);
+
+    return UnrealHttpRequestWithCallback(
+      UnrealHttpRequest(
+        url: '/remote/preset/$_transientPresetName/expose/property',
+        verb: 'PUT',
+        body: {
+          'ObjectPath': property!.property.objectPath,
+          'PropertyName': property.property.propertyName,
+          'EnableEditCondition': false,
+        },
+      ),
+      (response) => _onExposePropertyResponse(propertyId, response)
+          .onError((error, stackTrace) => _onExposePropertyError(error, stackTrace, propertyId, property)),
+    );
+  }
+
+  /// Create a request to the engine to unexpose a tracked property with the given [propertyId] on the transient
+  /// preset.
+  UnrealHttpRequestWithCallback _createUnexposePropertyRequest(TrackedPropertyId propertyId) {
+    final _TrackedProperty? property = _trackedProperties[propertyId];
+
+    assert(property != null);
+
+    return UnrealHttpRequestWithCallback(
+      UnrealHttpRequest(
+        url: '/remote/preset/$_transientPresetName/unexpose/property/${property!.label}',
+        verb: 'PUT',
+      ),
+      (response) => _onUnexposePropertyResponse(propertyId, response),
+    );
+  }
+
+  /// Called when we receive a response from the engine after exposing a property.
+  /// Returns a future which will provide true if the property was exposed, false if we already stopped tracking it,
+  /// and an error otherwise.
+  Future<void> _onExposePropertyResponse(TrackedPropertyId propertyId, UnrealHttpResponse response) {
+    final _TrackedProperty? property = _trackedProperties[propertyId];
+    if (property == null) {
+      // Property no longer exists. We may have already unregistered it.
+      return Future.value();
+    }
+
+    if (response.code != HttpResponseCode.ok) {
+      return Future.error('Response code ${response.code}');
+    }
+
+    final dynamic description = response.body['ExposedPropertyDescription'];
+    if (description == null) {
+      return Future.error('Response did not contain a property description');
+    }
+
+    // Get the name used to look up the property
+    final String? label = description['DisplayName'];
+    if (label == null) {
+      return Future.error('Response did not contain a property label');
+    }
+
+    // Get the name used to look up the property (unless we already have an override for it)
+    final dynamic underlyingProperty = description['UnderlyingProperty'];
+    String? typeName = property.property.typeNameOverride;
+    List<String>? enumValues;
+
+    if (typeName == null) {
+      // No override type name provided, so get the type from the data we received
+      final dynamic enumValuesData = underlyingProperty['Metadata']?['EnumValues'];
+      if (enumValuesData != null) {
+        // This is an enum type, so we can replace the specific enum type name provided by the engine with our internal
+        // generic enum type.
+        typeName = unrealEnumTypeName;
+        enumValues = enumValuesData.toString().split(', ');
+      } else {
+        // Just use the type provided by the engine
+        typeName = underlyingProperty?['Type'];
+      }
+
+      if (typeName == null) {
+        return Future.error('Response did not contain a property type');
+      }
+    }
+
+    if (!canConvertUnrealType(typeName)) {
+      return Future.error('Unsupported Unreal type $typeName');
+    }
+
+    // Get the initial value of the property
+    final dynamic rawValue = response.body['PropertyValues']?[0]?['PropertyValue'];
+    if (rawValue == null) {
+      return Future.error('Response did not contain a value');
+    }
+
+    // Get metadata if it was included
+    final dynamic responseMetaData = description['Metadata'];
+    final metaData = UnrealPropertyMetadata(
+      displayName: underlyingProperty?['DisplayName'] ?? '',
+      minValue: convertUnrealTypeJsonToDart(typeName, responseMetaData?['Min']),
+      maxValue: convertUnrealTypeJsonToDart(typeName, responseMetaData?['Max']),
+      enumValues: enumValues,
+    );
+
+    property.label = label;
+    property.value = convertUnrealTypeJsonToDart(
+      typeName,
+      rawValue,
+      conversionMetadata: property.conversionMetadata,
+    );
+    property.typeName = typeName;
+    property.metadata = metaData;
+
+    _propertyIdsByLabel[label] = propertyId;
+
+    // Complete the exposure future, then null it to indicate that this is already exposed.
+    property.exposeCompleter?.complete();
+    property.exposeCompleter = null;
+
+    return Future.value();
+  }
+
+  /// Called when an error happens while attempting to expose a property.
+  void _onExposePropertyError<T>(
+      T error, StackTrace stackTrace, TrackedPropertyId propertyId, _TrackedProperty property) {
+    _log.severe(
+        'Failed to expose property "${property.property.propertyName}" on ${property.property.objectPath}', error);
+
+    _trackedProperties.remove(propertyId);
+  }
+
+  /// Called when the engine confirms that a property has been unexposed.
+  void _onUnexposePropertyResponse(TrackedPropertyId propertyId, UnrealHttpResponse response) {
+    final _TrackedProperty? property = _trackedProperties[propertyId];
+    if (property == null) {
+      return;
+    }
+
+    property.bIsPendingUnexpose = false;
+
+    if (response.code != HttpResponseCode.ok) {
+      if (property.exposeCompleter != null) {
+        // Treat the property as re-exposed since we failed to remove it anyway
+        property.exposeCompleter!.complete();
+        property.exposeCompleter = null;
+        return;
+      }
+
+      /// TODO: Handle this case better. We should try again to remove it, but maybe with exponential backoff
+      return;
+    }
+
+    if (property.exposeCompleter != null) {
+      // We started tracking this property again after requesting unexpose, so we need to send a new expose request
+      // instead of removing it
+      _exposeProperty(propertyId);
+      return;
+    }
+
+    _propertyIdsByLabel.remove(property.label);
+    _trackedProperties.remove(propertyId);
+  }
+
+  /// Called each tick, when we're ready to send updates back to the engine.
+  void _onTick(Timer timer) {
+    final List<dynamic> messages = _finalizeTickChanges();
+
+    if (_transientPresetName == null) {
+      // We don't have a transient preset to send the message to, so try to create it and queue the messages to be sent
+      // once it's created.
+      _getOrCreateTransientPreset();
+      _messagesPendingPreset.addAll(messages);
+      return;
+    }
+
+    if (_messagesPendingPreset.isNotEmpty) {
+      // Update the waiting messages with the preset name
+      for (final dynamic message in _messagesPendingPreset) {
+        if (message['Parameters']?['PresetName'] != null) {
+          message['Parameters']['PresetName'] = _transientPresetName;
+        }
+      }
+
+      messages.insertAll(0, _messagesPendingPreset);
+    }
+
+    // Send the property update request
+    if (messages.isNotEmpty) {
+      _connectionManager.sendBatchedMessage(messages);
+    }
+
+    // Request expose/unexpose for any properties whose exposure states were changed
+    final List<UnrealHttpRequestWithCallback> httpRequests = _generatePropertyExposureChangeRequests();
+    if (httpRequests.isNotEmpty) {
+      _connectionManager.sendBatchedHttpRequestWithCallbacks(httpRequests);
+    }
+  }
+
+  /// Get the name of the transient preset we use to manage property subscriptions.
+  /// If the transient preset doesn't exist yet, this will send a creation request.
+  Future<String> _getOrCreateTransientPreset() {
+    if (_transientPresetNameFuture == null) {
+      // We don't already have a transient preset name (or one pending), so request a new one
+      final Future<UnrealHttpResponse> futureRequest =
+          _connectionManager.sendHttpRequest(const UnrealHttpRequest(url: '/remote/preset/transient', verb: 'PUT'));
+
+      _transientPresetNameFuture = futureRequest.then(
+        (response) {
+          if (response.code != HttpResponseCode.ok) {
+            return Future.error('Failed to create transient preset');
+          }
+
+          String? name = response.body?['Preset']?['Name'];
+          if (name == null) {
+            return Future.error('No preset name received from engine');
+          }
+
+          _onTransientPresetCreated(name);
+
+          return Future.value(name);
+        },
+      );
+    }
+
+    return _transientPresetNameFuture!;
+  }
+
+  /// Called when we receive the name of the transient preset requested from the engine.
+  void _onTransientPresetCreated(String name) {
+    _transientPresetName = name;
+
+    final List<dynamic> messages = [];
+
+    // Subscribe to property events for this preset, and mark it to be destroyed when we disconnect
+    messages.add(createUnrealWebSocketMessage('preset.register', {'PresetName': name}));
+    messages.add(createUnrealWebSocketMessage('preset.transient.autodestroy', {'PresetName': name}));
+
+    if (_messagesPendingPreset.isNotEmpty) {
+      // Update the preset name in all waiting messages
+      for (dynamic message in _messagesPendingPreset) {
+        if (message['Parameters']?['PresetName'] != null) {
+          message['Parameters']['PresetName'] = _transientPresetName;
+        }
+      }
+
+      // Send the messages that were waiting
+      messages.addAll(_messagesPendingPreset);
+      _messagesPendingPreset.clear();
+    }
+
+    _connectionManager.sendBatchedMessage(messages);
+  }
+
+  /// Returns the ID that should be used for the next tracked property.
+  TrackedPropertyId _getNextTrackedPropertyId() {
+    // The last tracked ID is no longer tracked, so we can reuse it
+    if (_lastTrackedId.isValid && !_trackedProperties.containsKey(_lastTrackedId)) {
+      return _lastTrackedId;
+    }
+
+    // Increment until we find an unused ID
+    final TrackedPropertyId startId = _lastTrackedId;
+    do {
+      _lastTrackedId = TrackedPropertyId(_lastTrackedId.id + 1);
+    } while (_trackedProperties.containsKey(_lastTrackedId) && _lastTrackedId != startId);
+
+    // If we've somehow exhausted all possible IDs, give up
+    if (_lastTrackedId == startId) {
+      throw Exception('All possible tracked property IDs exhausted');
+    }
+
+    return _lastTrackedId;
+  }
+
+  /// Get a tracked property if it's been exposed in the engine, or null if it isn't.
+  /// You should generally call this instead of directly indexing into [_trackedProperties] so that you don't access
+  /// values of a property that isn't ready.
+  _TrackedProperty? _getReadyTrackedProperty(TrackedPropertyId propertyId) {
+    final _TrackedProperty? property = _trackedProperties[propertyId];
+
+    if (property?.value == null) {
+      return null;
+    }
+
+    return property;
+  }
+
+  /// Convert all the value changes this tick into a list of messages to send to the server, and remember the change in
+  /// each corresponding property for future rollback/replay.
+  List<dynamic> _finalizeTickChanges() {
+    if (_transactionManager.activeTransactionId != _activeTransactionId) {
+      // If somebody else holds the active transaction, we can't make changes
+      assert(_activeTransactionId == UnrealTransactionManager.invalidId);
+      return [];
+    }
+
+    final List<dynamic> messages = [];
+
+    // Send changes from this tick
+    for (final _EngineTransaction transaction in _currentTickTransactions) {
+      if (transaction.changes.isNotEmpty) {
+        if (!transaction.bWasCreatedInEngine) {
+          // Create the transaction in the engine before we send its changes
+          final beginMessage = _transactionManager.createBeginTransactionMessage(
+            transaction.description,
+            // If the transaction ends prematurely, flag it to be re-created
+            prematureEndCallback: () => transaction.bWasCreatedInEngine = false,
+          );
+          assert(beginMessage != null);
+          messages.add(beginMessage);
+
+          _activeTransactionId = _transactionManager.activeTransactionId;
+          transaction.bWasCreatedInEngine = true;
+        }
+
+        // Add transaction changes
+        for (final MapEntry<TrackedPropertyId, _TrackedProperty> propertyEntry in _trackedProperties.entries) {
+          final List<_TrackedPropertyChange>? propertyChanges = transaction.changes[propertyEntry.key];
+
+          if (propertyChanges == null) {
+            continue;
+          }
+
+          final _TrackedProperty property = propertyEntry.value;
+          property.recentTickChanges.add(_SequencedPropertyChangeRecord(
+            sequenceNumber: _sequenceNumber,
+            changes: List.from(propertyChanges),
+          ));
+
+          // Create a request for each change
+          for (final _TrackedPropertyChange change in propertyChanges) {
+            // Note that the preset name is "TEMP" if we don't yet have a transient preset. We'll update this for any
+            // queued messages before sending them (in [_onTransientPresetCreated])
+            messages.add(_createPropertyChangeMessage(
+              propertyEntry.value,
+              change,
+              _activeTransactionId,
+              _transientPresetName ?? 'TEMP',
+            ));
+          }
+        }
+      }
+
+      if (!transaction.bIsActive && transaction.bWasCreatedInEngine) {
+        // Transaction is no longer active, so end it in the engine
+        final endMessage = _transactionManager.createEndTransactionMessage();
+        assert(endMessage != null);
+        messages.add(endMessage);
+
+        _activeTransactionId = UnrealTransactionManager.invalidId;
+      }
+    }
+
+    if (messages.isNotEmpty) {
+      ++_sequenceNumber;
+    }
+
+    // Clear out the list of inactive transactions
+    _currentTickTransactions.removeWhere((_EngineTransaction transaction) => !transaction.bIsActive);
+
+    // Only one active transaction should exist at a time
+    assert(_currentTickTransactions.length <= 1);
+
+    // Clear out changes we just sent for the active transaction
+    if (_currentTickTransactions.isNotEmpty) {
+      _currentTickTransactions[0].changes.clear();
+    }
+
+    return messages;
+  }
+
+  /// Generate a list of HTTP requests to expose/unexpose pending properties on the transient preset.
+  List<UnrealHttpRequestWithCallback> _generatePropertyExposureChangeRequests() {
+    final List<UnrealHttpRequestWithCallback> requests = [];
+
+    // Add messages for any properties that need to be exposed/unexposed
+    for (final TrackedPropertyId propertyId in _propertiesPendingExpose) {
+      final _TrackedProperty? property = _trackedProperties[propertyId];
+      assert(property != null);
+      assert(property!.exposeCompleter != null);
+
+      requests.add(_createExposePropertyRequest(propertyId));
+    }
+
+    for (final TrackedPropertyId propertyId in _propertiesPendingUnexpose) {
+      final _TrackedProperty? property = _trackedProperties[propertyId];
+      assert(property != null);
+      assert(property!.bIsPendingUnexpose);
+
+      requests.add(_createUnexposePropertyRequest(propertyId));
+    }
+
+    _propertiesPendingExpose.clear();
+    _propertiesPendingUnexpose.clear();
+
+    return requests;
+  }
+
+  /// Create a WebSocket message to apply [change] to [property], as part of the transaction with ID [transactionId],
+  /// via the preset named [presetName].
+  dynamic _createPropertyChangeMessage(
+      _TrackedProperty property, _TrackedPropertyChange change, int transactionId, String presetName) {
+    final dynamic parameters = change.operation.makeMessageParameters(PropertyModifyEvent(
+      property: property.property,
+      currentValue: property.value,
+      deltaValue: change.deltaValue,
+      presetName: presetName,
+      presetPropertyLabel: property.label,
+      typeName: property.typeName,
+      sequenceNumber: _sequenceNumber,
+      transactionId: transactionId,
+    ));
+
+    return createUnrealWebSocketMessage(change.operation.messageName, parameters);
+  }
+
+  /// Called when we get a message from the engine that our transient preset's fields have changed.
+  void _onPresetFieldsChanged(dynamic message) {
+    final List<dynamic>? changedFields = message['ChangedFields'];
+    if (changedFields == null) {
+      _log.warning('Received a changed fields message with no ChangedFields property. It will be ignored.');
+      return;
+    }
+
+    final String? engineSequenceNumberString = message['SequenceNumber'];
+    if (engineSequenceNumberString == null) {
+      _log.warning('Received a changed fields message with no sequence number. It will be ignored.');
+      return;
+    }
+
+    final int? engineSequenceNumber = int.tryParse(engineSequenceNumberString);
+    if (engineSequenceNumber == null) {
+      _log.warning('Received a changed fields message but failed to parse the sequence number. It will be ignored.');
+      return;
+    }
+
+    for (dynamic changeData in changedFields) {
+      _onSinglePresetFieldChanged(changeData, engineSequenceNumber);
+    }
+  }
+
+  /// Called for each field that changed in a PresetFieldsChanged message from the engine.
+  void _onSinglePresetFieldChanged(dynamic changeData, int engineSequenceNumber) {
+    final String? propertyLabel = changeData['PropertyLabel'];
+    if (propertyLabel == null) {
+      _log.warning('Received a changed fields message with no property label');
+      return;
+    }
+
+    final TrackedPropertyId? propertyId = _propertyIdsByLabel[propertyLabel];
+    if (propertyId == null) {
+      _log.warning('Received a changed fields message, but property $propertyLabel is not tracked');
+      return;
+    }
+
+    final _TrackedProperty? property = _trackedProperties[propertyId];
+    assert(property != null, 'A property in _propertyIdsByLabel should also be in _trackedProperties');
+
+    // Parse the corrected value from the JSON data
+    final dynamic engineValue = convertUnrealTypeJsonToDart(
+      property!.typeName,
+      changeData['PropertyValue'],
+      previousValue: property.value,
+      conversionMetadata: property.conversionMetadata,
+    );
+    if (engineValue == null) {
+      _log.warning('Failed to read PropertyValue for change to $propertyLabel');
+      return;
+    }
+
+    // Start at the time given by the server, then reapply our own input since then
+    final List<_SequencedPropertyChangeRecord> recentTickChanges = property.recentTickChanges;
+
+    // First, skip any changes we recorded that were already handled by the server
+    int lastSkippedChangeIndex = -1;
+    for (int i = 0; i < recentTickChanges.length; ++i) {
+      if (recentTickChanges[i].sequenceNumber <= engineSequenceNumber) {
+        lastSkippedChangeIndex = i;
+      } else {
+        break;
+      }
+    }
+
+    recentTickChanges.removeRange(0, lastSkippedChangeIndex + 1);
+
+    // Apply any remaining changes sequentially
+    dynamic newValue = engineValue;
+
+    for (_SequencedPropertyChangeRecord changeRecord in recentTickChanges) {
+      for (_TrackedPropertyChange change in changeRecord.changes) {
+        newValue = change.apply(newValue);
+      }
+    }
+
+    // Also apply current tick's changes if applicable
+    final _EngineTransaction? activeTransaction = _getActiveTransaction();
+    if (activeTransaction != null) {
+      final List<_TrackedPropertyChange>? currentChanges = activeTransaction.changes[propertyId];
+      if (currentChanges != null) {
+        for (_TrackedPropertyChange change in currentChanges) {
+          newValue = change.apply(newValue);
+        }
+      }
+    }
+
+    property.value = newValue;
+  }
+
+  /// Called when the connection manager's connection state changes.
+  void _onConnectionStateChanged(EngineConnectionState connectionState) {
+    switch (connectionState) {
+      case EngineConnectionState.connected:
+        _onConnect();
+        break;
+
+      case EngineConnectionState.disconnected:
+        _onDisconnect();
+        break;
+    }
+  }
+
+  /// Called when the connection manager has connected to the engine.
+  void _onConnect() {
+    _transientPresetNameFuture = _getOrCreateTransientPreset();
+
+    _tickTimer = Timer.periodic(_tickRate, _onTick);
+  }
+
+  /// Called when the connection manager has disconnected from the engine.
+  void _onDisconnect() {
+    _tickTimer?.cancel();
+
+    _tickTimer = null;
+    _trackedProperties.clear();
+    _propertyIdsByLabel.clear();
+    _currentTickTransactions.clear();
+    _messagesPendingPreset.clear();
+    _transientPresetName = null;
+    _transientPresetNameFuture = null;
+    _lastTrackedId = const TrackedPropertyId.invalid();
+  }
+}
+
+/// An ID used to refer to a property being tracked by EnginePropertyManager.
+class TrackedPropertyId {
+  static const _invalidId = -1;
+
+  /// Create a tracked property ID with the given numerical value.
+  const TrackedPropertyId(this.id);
+
+  /// Create an invalid property ID.
+  const TrackedPropertyId.invalid() : id = _invalidId;
+
+  /// The internal int representation of the ID.
+  final int id;
+
+  @override
+  bool operator ==(other) {
+    return other is TrackedPropertyId && other.id == id;
+  }
+
+  @override
+  int get hashCode => id;
+
+  /// Returns whether this ID refers to a valid property.
+  bool get isValid => id != _invalidId;
+
+  @override
+  String toString() => isValid ? '$id' : 'INVALID ($id)';
+}
+
+/// Metadata about a property received from the engine.
+/// Any data in this class is exposed publicly via [UnrealPropertyManager.getTrackedPropertyMetaData].
+class UnrealPropertyMetadata {
+  const UnrealPropertyMetadata({this.displayName = '', this.minValue, this.maxValue, this.enumValues});
+
+  /// The in-engine name of the property, if provided.
+  final String displayName;
+
+  /// The property's minimum value, if any.
+  final dynamic minValue;
+
+  /// The property's maximum value, if any.
+  final dynamic maxValue;
+
+  /// The list of possible enum values, if provided.
+  final List<String>? enumValues;
+}
+
+/// Information about a transaction we created (or will create) in the engine.
+class _EngineTransaction {
+  _EngineTransaction({required this.description});
+
+  /// The description of the transaction displayed in the editor.
+  final String description;
+
+  /// Whether the transaction is currently active in the app.
+  bool bIsActive = true;
+
+  /// Whether the transaction has been created in the engine.
+  bool bWasCreatedInEngine = false;
+
+  /// Changes that happened to each property as part of this transaction.
+  /// The changes are listed chronologically and may be combined if they used the same type of operation.
+  final Map<TrackedPropertyId, List<_TrackedPropertyChange>> changes = {};
+}
+
+/// Data about a property being tracked by EnginePropertyManager.
+class _TrackedProperty {
+  _TrackedProperty(
+      {required this.property, required this.conversionMetadata, List<PropertyChangeCallback>? initialCallbacks}) {
+    if (initialCallbacks != null) {
+      changeCallbacks.addAll(initialCallbacks);
+    }
+  }
+
+  /// Data used to find the property in the engine.
+  final UnrealProperty property;
+
+  /// Functions to call when this property's value changes.
+  final Set<PropertyChangeCallback> changeCallbacks = {};
+
+  /// Changes that have happened in each tick since the last engine update was received for this property.
+  /// The last entry in this list is the most recently completed tick.
+  final List<_SequencedPropertyChangeRecord> recentTickChanges = [];
+
+  /// The label of the property as exposed on the transient preset.
+  late String label;
+
+  /// The type name of the property provided by the engine.
+  late String typeName;
+
+  /// Metadata about the property provided by the engine..
+  late UnrealPropertyMetadata metadata;
+
+  /// If true, we've sent a message requesting that the property be unexposed in the engine.
+  bool bIsPendingUnexpose = false;
+
+  /// Completer for when this property is finished being exposed. If this is non-null, the property is pending expose
+  /// (or re-expose if bIsPendingUnexpose is also true).
+  Completer<void>? exposeCompleter;
+
+  /// Arbitrary data that will be passed to conversion functions.
+  Map<String, dynamic> conversionMetadata;
+
+  /// Underlying predicted value of the property. Use the [value] instead so callbacks are triggered.
+  dynamic _value;
+
+  /// The current predicted value of the property.
+  dynamic get value => _value;
+
+  set value(dynamic newValue) {
+    bool bWasReady = _value != null;
+    _value = newValue;
+
+    // If we weren't ready before, then this is the first time the value was set. Callers are expecting to receive
+    // confirmation that the property was exposed (including its value) rather than a callback containing the value,
+    // so don't send callbacks yet.
+    if (bWasReady) {
+      for (PropertyChangeCallback callback in changeCallbacks) {
+        callback(newValue);
+      }
+    }
+  }
+}
+
+/// A single change to a tracked property.
+class _TrackedPropertyChange {
+  _TrackedPropertyChange({required this.operation, this.deltaValue});
+
+  /// The operation used to apply the change.
+  final PropertyModifyOperation operation;
+
+  /// The amount by which the property's value was changed.
+  dynamic deltaValue;
+
+  /// Returns whether this can be combined with another operation.
+  bool canCombine(_TrackedPropertyChange other) {
+    return other.operation == operation;
+  }
+
+  /// Combine the delta values of two operations if possible. The combined value will be stored in this change.
+  void combine(_TrackedPropertyChange other) {
+    if (!canCombine(other)) {
+      throw Exception('Unable to combine tracked property changes');
+    }
+
+    deltaValue = other.apply(deltaValue);
+  }
+
+  /// Returns the result of applying the operation to the given [value].
+  dynamic apply(dynamic value) => operation.apply(value, deltaValue);
+
+  /// Given a [referenceValue] that the change will be applied to, modify the change so that it doesn't exceed the
+  /// provided [min] and [max] values by clamping it. If [min] or [max] are null, that constraint won't be applied.
+  /// Returns true if the change is still necessary, and false if after clamping, there would effectively be no change.
+  bool clamp({required dynamic referenceValue, dynamic min, dynamic max}) {
+    final dynamic clampedDeltaValue = operation.clamp(referenceValue, deltaValue, min, max);
+
+    if (clampedDeltaValue != null) {
+      deltaValue = clampedDeltaValue;
+      return true;
+    }
+
+    return operation.bIsAlwaysSignificant;
+  }
+
+  /// Given a [referenceValue] that the change will be applied to, modify the change so that it doesn't exceed the
+  /// provided [min] and [max] by looping around them.
+  /// Returns true if loop was successful.
+  bool loop({required dynamic referenceValue, dynamic min, dynamic max}) {
+    final dynamic loopedDeltaValue = operation.loop(referenceValue, deltaValue, min, max);
+
+    if (loopedDeltaValue != null) {
+      deltaValue = loopedDeltaValue;
+    }
+
+    return loopedDeltaValue != null;
+  }
+}
+
+/// All changes that occurred to a property on a given tick.
+class _SequencedPropertyChangeRecord {
+  const _SequencedPropertyChangeRecord({required this.sequenceNumber, required this.changes});
+
+  /// The sequence number when the tick was completed.
+  final int sequenceNumber;
+
+  /// The changes that occurred to the property on the tick.
+  final List<_TrackedPropertyChange> changes;
+}
diff --git a/Engine/Extras/VirtualProduction/EpicStageApp/lib/models/unreal_transaction_manager.dart b/Engine/Extras/VirtualProduction/EpicStageApp/lib/models/unreal_transaction_manager.dart
new file mode 100644
--- /dev/null
+++ b/Engine/Extras/VirtualProduction/EpicStageApp/lib/models/unreal_transaction_manager.dart
@@ -0,0 +1,152 @@
+// Copyright Epic Games, Inc. All Rights Reserved.
+
+import 'package:flutter/material.dart';
+import 'package:logging/logging.dart';
+import 'package:provider/provider.dart';
+
+import '../utilities/unreal_utilities.dart';
+import 'engine_connection.dart';
+
+final _log = Logger('UnrealTransactionManager');
+
+/// Manages the state of transactions in the Unreal Engine.
+class UnrealTransactionManager {
+  /// Construct a property manager and register it with the connection manager.
+  UnrealTransactionManager(BuildContext context)
+      : _connectionManager = Provider.of<EngineConnectionManager>(context, listen: false) {
+    _connectionManager.registerMessageListener('TransactionEnded', _onEngineTransactionEnded);
+  }
+
+  /// Transaction ID returned when a transaction could not be created.
+  static const int invalidId = -1;
+
+  /// The connection manager we'll use to communicate with the engine.
+  final EngineConnectionManager _connectionManager;
+
+  /// The ID of the next transaction to be created.
+  int _nextTransactionId = 0;
+
+  /// Data for the current active transaction, or null if there isn't one.
+  _EngineTransaction? _activeTransaction;
+
+  /// Check whether there's a transaction in progress.
+  bool get bHasActiveTransaction => _activeTransaction != null;
+
+  /// The ID of the active transaction, or [invalidId] if there isn't one.
+  int get activeTransactionId => _activeTransaction?.id ?? invalidId;
+
+  /// Start a transaction.
+  /// Returns true if the transaction was successfully created.
+  bool beginTransaction(String description, {void Function()? prematureEndCallback}) {
+    final message = createBeginTransactionMessage(description, prematureEndCallback: prematureEndCallback);
+    if (message == null) {
+      return false;
+    }
+
+    _connectionManager.sendRawMessage(message);
+    return true;
+  }
+
+  /// End the current transaction.
+  /// Returns false if there was no active transaction to end.
+  bool endTransaction() {
+    final message = createEndTransactionMessage();
+    if (message == null) {
+      return false;
+    }
+
+    _connectionManager.sendRawMessage(message);
+    return true;
+  }
+
+  /// Create a message to start a transaction.
+  /// The caller is expected to send this message using [EngineConnectionManager], e.g. as part of a batch.
+  /// Returns null if there was already an active transaction.
+  dynamic createBeginTransactionMessage(String description, {void Function()? prematureEndCallback}) {
+    if (_activeTransaction != null) {
+      _log.warning('Tried to start transaction, but one was already in progress');
+      return null;
+    }
+
+    _activeTransaction = _EngineTransaction(
+      id: _nextTransactionId,
+      prematureEndCallback: prematureEndCallback,
+    );
+
+    final message = createUnrealWebSocketMessage('transaction.begin', {
+      'Description': description,
+      'TransactionId': _nextTransactionId,
+    });
+
+    ++_nextTransactionId;
+
+    return message;
+  }
+
+  /// Create a message to end a transaction.
+  /// The caller is expected to send this message using [EngineConnectionManager], e.g. as part of a batch.
+  /// Returns null if there was no active transaction to end.
+  dynamic createEndTransactionMessage() {
+    if (_activeTransaction == null) {
+      _log.warning('Tried to end transaction, but no transaction was in progress');
+      return false;
+    }
+
+    final message = createUnrealWebSocketMessage('transaction.end', {
+      'TransactionId': _activeTransaction!.id,
+    });
+
+    _activeTransaction = null;
+    return message;
+  }
+
+  /// Helper function to automatically create and end a transaction before/after calling a function, taking into account
+  /// whether the engine version allows a manual transaction for property changes.
+  /// If a manual transaction is allowed, this will create a transaction with the given [description], it will call
+  /// [sendRequest] passing true for [bIsManualTransaction], then end the transaction.
+  /// Otherwise, it will just call [sendRequest], passing false for [bIsManualTransaction].
+  dynamic wrapWithTransactionIfManualAllowedForProperties(
+    String description,
+    Function(bool bIsManualTransaction) sendRequest,
+  ) {
+    final bool bIsManualTransaction = _connectionManager.apiVersion?.bCanHttpSetPropertyInManualTransaction == true;
+    if (bIsManualTransaction) {
+      if (!beginTransaction(description)) {
+        _log.warning('Failed to begin transaction for "$description"');
+        return;
+      }
+    }
+
+    final result = sendRequest(bIsManualTransaction);
+
+    if (bIsManualTransaction) {
+      endTransaction();
+    }
+
+    return result;
+  }
+
+  /// Called when the engine reports that a transaction has ended.
+  void _onEngineTransactionEnded(dynamic message) {
+    final int? id = message['TransactionId'];
+    if (id == null) {
+      _log.warning('Received TransactionEnded event with no transaction ID');
+      return;
+    }
+
+    if (_activeTransaction?.id == id) {
+      _activeTransaction?.prematureEndCallback?.call();
+    }
+  }
+}
+
+/// Data about an active transaction that we created in the engine.
+class _EngineTransaction {
+  _EngineTransaction({required this.id, this.prematureEndCallback});
+
+  /// The ID used to the transaction in the engine.
+  final int id;
+
+  /// Callback function if this transaction ends prematurely.
+  final void Function()? prematureEndCallback;
+}
