Branch snapshot for CL 27138787

diff --git a/Engine/Extras/VirtualProduction/EpicStageApp/lib/models/unreal_types.dart b/Engine/Extras/VirtualProduction/EpicStageApp/lib/models/unreal_types.dart
new file mode 100644
--- /dev/null
+++ b/Engine/Extras/VirtualProduction/EpicStageApp/lib/models/unreal_types.dart
@@ -0,0 +1,239 @@
+// Copyright Epic Games, Inc. All Rights Reserved.
+
+import 'dart:collection';
+
+import 'package:flutter/material.dart';
+import 'package:logging/logging.dart';
+
+import 'actor_data/per_class_actor_data.dart';
+import 'unreal_actor_manager.dart';
+import 'unreal_class_data.dart';
+
+final _log = Logger('UnrealTypes');
+
+/// Data about an object that exists in an Unreal Engine instance.
+class UnrealObject extends ChangeNotifier {
+  UnrealObject({required String path, required String name})
+      : _name = name,
+        _path = path;
+
+  /// The full path of the object in the editor.
+  final String _path;
+
+  /// Classes to which the engine has reported this class belongs.
+  final Set<String> _classNames = {};
+
+  /// The lowest-level Unreal class that this belongs to based on its class names.
+  UnrealClass? _unrealClass;
+
+  /// Extra data retrieved from the engine based on this actor's classes.
+  final Map<Type, UnrealPerClassActorData> _perClassData = {};
+
+  /// The object's display name.
+  String _name;
+
+  /// Whether this object has been destroyed in the engine.
+  bool _bIsDestroyed = false;
+
+  /// Whether this object is no longer being watched, so may be out of date.
+  bool _bIsStale = false;
+
+  /// The full path of the object in the editor.
+  String get path => _path;
+
+  /// The object's display name.
+  String get name => _name;
+
+  set name(String value) {
+    if (value != _name) {
+      _name = value;
+      notifyListeners();
+    }
+  }
+
+  /// Whether this object has been destroyed in the engine.
+  bool get bIsDestroyed => _bIsDestroyed;
+
+  /// Whether this object is no longer being watched, so may be out of date.
+  bool get bIsStale => _bIsStale;
+
+  /// Whether this object still exists in the engine and is actively being watched.
+  /// In either case, this object will never be valid again, and you should get a new reference from the
+  /// [UnrealActorManager] once the actor is watched/created again.
+  bool get bIsValid => !bIsDestroyed && !bIsStale;
+
+  /// Classes to which the engine has reported this class belongs.
+  UnmodifiableSetView<String> get classNames => UnmodifiableSetView(_classNames);
+
+  /// The lowest-level Unreal class that this belongs to based on its class names.
+  UnrealClass? get unrealClass => _unrealClass;
+
+  /// Return true if the the engine has reported the actor as a member of a [className].
+  /// Note that this does not handle inheritance; it will only check exact matches in list of classes for which
+  /// [UnrealActorManager] is actively listening.
+  bool isA(String className) {
+    return _classNames.contains(className);
+  }
+
+  /// Return true if the the engine has reported the actor as a member of any of the given [classNames].
+  /// Note that this does not handle inheritance; it will only check exact matches in list of classes for which
+  /// [UnrealActorManager] is actively listening.
+  bool isAny(List<String> classNames) {
+    return classNames.any((className) => _classNames.contains(className));
+  }
+
+  /// Indicate that this actor has been destroyed.
+  void onDestroyed() {
+    if (!_bIsDestroyed) {
+      _bIsDestroyed = true;
+      notifyListeners();
+    }
+  }
+
+  /// Indicate that this actor is no longer being watched.
+  void onUnwatched() {
+    if (!_bIsStale) {
+      _bIsStale = true;
+      notifyListeners();
+    }
+  }
+
+  /// Add a class name that the actor belongs to.
+  void addClassName(String className) {
+    _classNames.add(className);
+    _unrealClass = UnrealClassRegistry.getClassForObject(this);
+    notifyListeners();
+  }
+
+  /// Add class-specific data to the actor.
+  void addPerClassData(dynamic data) {
+    if (_perClassData[data.runtimeType] != null) {
+      _log.warning('Tried to add ${data.runtimeType} to $this, but it already had data of that type');
+      return;
+    }
+
+    _perClassData[data.runtimeType] = data;
+    data.addListener(notifyListeners);
+    notifyListeners();
+  }
+
+  /// Get class-specific data of the given type from the actor (if it exists).
+  T? getPerClassData<T extends UnrealPerClassActorData>() {
+    return _perClassData[T] as T?;
+  }
+
+  /// Get class-specific data of the given type from the actor (if it exists).
+  UnrealPerClassActorData? getPerClassDataOfType(Type type) {
+    return _perClassData[type];
+  }
+
+  /// Get the path of the icon associated with this actor's type.
+  String? getIconPath() {
+    return _unrealClass?.getIconPathForInstance(this);
+  }
+
+  /// Get the color associated with this actor's type.
+  Color? getUIColor() {
+    return _unrealClass?.getUIColorForInstance(this);
+  }
+
+  @override
+  void dispose() {
+    for (final UnrealPerClassActorData data in _perClassData.values) {
+      data.dispose();
+    }
+    _perClassData.clear();
+
+    super.dispose();
+  }
+}
+
+/// Data about a property on an object that exists in an Unreal Engine instance.
+class UnrealProperty {
+  const UnrealProperty({
+    required this.objectPath,
+    required this.propertyName,
+    this.typeNameOverride,
+    String? friendlyObjectName,
+  }) : _friendlyObjectName = friendlyObjectName;
+
+  /// A reference to no property.
+  static const UnrealProperty empty = UnrealProperty(objectPath: '', propertyName: '');
+
+  /// The full path of object in the editor.
+  final String objectPath;
+
+  /// The property's name.
+  final String propertyName;
+
+  /// An optional override for the property's type, which lets the app handle the property differently than
+  /// the default for its actual type in the engine (e.g. interpret a Vector4 as a color instead).
+  final String? typeNameOverride;
+
+  /// The optional user-friendly name of the object the property belongs to.
+  final String? _friendlyObjectName;
+
+  /// Whether this references a property (though that property may not exist).
+  bool get bIsNotEmpty => objectPath.isNotEmpty && propertyName.isNotEmpty;
+
+  /// Get the name of the object to display to the user.
+  String get objectName {
+    if (_friendlyObjectName != null) {
+      return _friendlyObjectName!;
+    }
+
+    final int dotIndex = objectPath.lastIndexOf('.');
+    if (dotIndex == -1) {
+      return objectPath;
+    }
+
+    return objectPath.substring(dotIndex + 1);
+  }
+
+  /// The name of the last property in the chain if this is a nested property.
+  String get lastPropertyName {
+    final lastDotIndex = propertyName.lastIndexOf('.');
+    return (lastDotIndex == -1) ? propertyName : propertyName.substring(lastDotIndex + 1);
+  }
+
+  /// Make a property with a suffix applied to the [propertyName].
+  UnrealProperty makeSubproperty(String suffix, {String? typeNameOverride, String? friendlyObjectName}) {
+    return UnrealProperty(
+      objectPath: objectPath,
+      propertyName: '$propertyName.$suffix',
+      typeNameOverride: typeNameOverride,
+      friendlyObjectName: friendlyObjectName,
+    );
+  }
+
+  @override
+  bool operator ==(Object other) {
+    if (identical(this, other)) {
+      return true;
+    }
+
+    if (other.runtimeType != runtimeType) {
+      return false;
+    }
+
+    return other is UnrealProperty &&
+        other.objectPath == objectPath &&
+        other.propertyName == propertyName &&
+        other.typeNameOverride == typeNameOverride &&
+        other._friendlyObjectName == _friendlyObjectName;
+  }
+
+  @override
+  int get hashCode => Object.hash(objectPath, propertyName, typeNameOverride, _friendlyObjectName);
+}
+
+/// Data about an actor template available in the engine.
+class UnrealTemplateData {
+  const UnrealTemplateData({required this.name, required this.path});
+
+  /// The user-friendly name of the template.
+  final String name;
+
+  /// The path of the template asset.
+  final String path;
+}
diff --git a/Engine/Extras/VirtualProduction/EpicStageApp/lib/routes.dart b/Engine/Extras/VirtualProduction/EpicStageApp/lib/routes.dart
new file mode 100644
--- /dev/null
+++ b/Engine/Extras/VirtualProduction/EpicStageApp/lib/routes.dart
@@ -0,0 +1,32 @@
+// Copyright Epic Games, Inc. All Rights Reserved.
+
+import 'widgets/screens/eula/eula_screen.dart';
+import 'package:flutter/widgets.dart';
+
+import 'widgets/screens/connect/connect.dart';
+import 'widgets/screens/main/stage_app_main_screen.dart';
+import 'widgets/screens/reconnect_screen.dart';
+
+/// Contains data about a route in the app
+class RouteData {
+  const RouteData({required this.createScreen});
+
+  /// A function to create the screen widget for the route.
+  final Widget Function(BuildContext context) createScreen;
+
+  /// A list of all routes in the app.
+  static final Map<String, RouteData> allRoutes = {
+    ConnectScreen.route: RouteData(
+      createScreen: (context) => const ConnectScreen(),
+    ),
+    ReconnectScreen.route: RouteData(
+      createScreen: (context) => const ReconnectScreen(),
+    ),
+    StageAppMainScreen.route: RouteData(
+      createScreen: (context) => const StageAppMainScreen(),
+    ),
+    EulaScreen.route: RouteData(createScreen: (context) => const EulaScreen())
+  };
+}
+
+final Map<String, String> routes = {};
diff --git a/Engine/Extras/VirtualProduction/EpicStageApp/lib/utilities/constants.dart b/Engine/Extras/VirtualProduction/EpicStageApp/lib/utilities/constants.dart
new file mode 100644
--- /dev/null
+++ b/Engine/Extras/VirtualProduction/EpicStageApp/lib/utilities/constants.dart
@@ -0,0 +1,62 @@
+// Copyright Epic Games, Inc. All Rights Reserved.
+
+import 'package:flutter/material.dart';
+
+/// Radius of the corners for top-level UI elements.
+const double outerCornerRadius = 16;
+
+/// Margin between the cards in a tab.
+const double cardMargin = 8;
+
+/// Radius between sections within cards.
+const double sectionMargin = 4;
+
+/// Radius of the corners for cards in a tab.
+const double cardCornerRadius = 8;
+
+/// Padding for scrollable list views within cards.
+const EdgeInsetsGeometry cardListViewPadding = EdgeInsets.only(top: 2);
+
+/// The Unreal class for nDisplay root actors.
+const String nDisplayRootActorClassName = '/Script/DisplayCluster.DisplayClusterRootActor';
+
+/// The Unreal class name used for light cards.
+const String lightCardClassName = '/Script/DisplayCluster.DisplayClusterLightCardActor';
+
+/// The Unreal class name used for chroma key cards.
+const String chromakeyCardClassName = '/Script/DisplayCluster.DisplayClusterChromakeyCardActor';
+
+/// The Unreal class names used for color correct windows.
+/// The preferred class name must be FIRST in this list.
+const List<String> colorCorrectWindowClassNames = [
+  '/Script/ColorCorrectRegions.ColorCorrectionWindow',
+  '/Script/ColorCorrectRegions.ColorCorrectWindow',
+];
+
+/// The Unreal class name used for color correct regions.
+const String colorCorrectRegionClassName = '/Script/ColorCorrectRegions.ColorCorrectRegion';
+
+/// The Unreal class name used for post-process volumes.
+const String postProcessVolumeClassName = '/Script/Engine.PostProcessVolume';
+
+/// List of names for scene actor classes that should be controllable from the app.
+List<String> get controllableClassNames {
+  if (_controllableClassNames == null) {
+    _controllableClassNames = [lightCardClassName, chromakeyCardClassName, colorCorrectRegionClassName];
+    _controllableClassNames!.addAll(colorCorrectWindowClassNames);
+  }
+
+  return _controllableClassNames!;
+}
+
+List<String>? _controllableClassNames = null;
+
+/// Default decoration style for text inputs.
+final InputDecoration collapsedInputDecoration = const InputDecoration(
+  isCollapsed: true,
+  contentPadding: EdgeInsets.symmetric(horizontal: 12, vertical: 16),
+  border: OutlineInputBorder(
+    borderRadius: BorderRadius.all(Radius.circular(4)),
+    borderSide: BorderSide.none,
+  ),
+);
diff --git a/Engine/Extras/VirtualProduction/EpicStageApp/lib/utilities/debug_utilities.dart b/Engine/Extras/VirtualProduction/EpicStageApp/lib/utilities/debug_utilities.dart
new file mode 100644
--- /dev/null
+++ b/Engine/Extras/VirtualProduction/EpicStageApp/lib/utilities/debug_utilities.dart
@@ -0,0 +1,18 @@
+// Copyright Epic Games, Inc. All Rights Reserved.
+
+import 'package:logging/logging.dart';
+
+import '../models/navigator_keys.dart';
+import '../widgets/elements/modal.dart';
+
+final _log = Logger('DebugAlert');
+
+/// Show an alert to the user with the given message.
+void showDebugAlert(String message) {
+  if (rootNavigatorKey.currentContext == null) {
+    return;
+  }
+
+  _log.info(message);
+  InfoModalDialog.show(message);
+}
diff --git a/Engine/Extras/VirtualProduction/EpicStageApp/lib/utilities/external_notifier.dart b/Engine/Extras/VirtualProduction/EpicStageApp/lib/utilities/external_notifier.dart
new file mode 100644
--- /dev/null
+++ b/Engine/Extras/VirtualProduction/EpicStageApp/lib/utilities/external_notifier.dart
@@ -0,0 +1,11 @@
+// Copyright Epic Games, Inc. All Rights Reserved.
+
+import 'package:flutter/material.dart';
+
+/// An extension of ChangeNotifier that allows [notifyListeners] to be called externally.
+class ExternalNotifier extends ChangeNotifier {
+  void notifyListeners() {
+    // This function is normally protected, but we want to make it public for this class.
+    super.notifyListeners();
+  }
+}
diff --git a/Engine/Extras/VirtualProduction/EpicStageApp/lib/utilities/guarded_refresh_state.dart b/Engine/Extras/VirtualProduction/EpicStageApp/lib/utilities/guarded_refresh_state.dart
new file mode 100644
--- /dev/null
+++ b/Engine/Extras/VirtualProduction/EpicStageApp/lib/utilities/guarded_refresh_state.dart
@@ -0,0 +1,19 @@
+// Copyright Epic Games, Inc. All Rights Reserved.
+
+import 'package:flutter/material.dart';
+
+/// Adds a function to refresh the widget only if it's still mounted.
+mixin GuardedRefreshState<T extends StatefulWidget> on State<T> {
+  /// Refresh the widget if it's still mounted.
+  void guardedRefresh() {
+    if (mounted) {
+      setState(() {});
+    }
+  }
+
+  /// Convenience version of the function for listening to a subscription stream, which passes in a value that we don't
+  /// care about for this function's purpose.
+  void refreshOnData(_) {
+    guardedRefresh();
+  }
+}
diff --git a/Engine/Extras/VirtualProduction/EpicStageApp/lib/utilities/json_utilities.dart b/Engine/Extras/VirtualProduction/EpicStageApp/lib/utilities/json_utilities.dart
new file mode 100644
--- /dev/null
+++ b/Engine/Extras/VirtualProduction/EpicStageApp/lib/utilities/json_utilities.dart
@@ -0,0 +1,15 @@
+// Copyright Epic Games, Inc. All Rights Reserved.
+
+/// Convert an enum to a value for JSON serialization.
+String enumToJsonValue<T extends Enum>(T value) {
+  return value.name;
+}
+
+/// Convert a serialized JSON value to an enum value.
+T? jsonToEnumValue<T extends Enum>(dynamic value, List<T> enumValues) {
+  if (!(value is String)) {
+    return null;
+  }
+
+  return enumValues.firstWhere((enumValue) => enumValue.name == value, orElse: null);
+}
diff --git a/Engine/Extras/VirtualProduction/EpicStageApp/lib/utilities/logging.dart b/Engine/Extras/VirtualProduction/EpicStageApp/lib/utilities/logging.dart
new file mode 100644
--- /dev/null
+++ b/Engine/Extras/VirtualProduction/EpicStageApp/lib/utilities/logging.dart
@@ -0,0 +1,177 @@
+// Copyright Epic Games, Inc. All Rights Reserved.
+
+import 'dart:io';
+
+import 'package:intl/intl.dart';
+import 'package:logging/logging.dart';
+import 'package:path_provider/path_provider.dart';
+
+final _timeFormatter = DateFormat('HH:mm:ss.sss');
+final _fileDateFormatter = DateFormat('yyyy_MM_dd');
+final _friendlyDateFormatter = DateFormat.yMMMMd();
+final _logDateFromName = RegExp(r'epic_stage_app_(.+)\.log');
+
+final _log = Logger('Logging');
+
+/// Color tokens to prepend to logs at each level.
+final Map<Level, String> _logColorTokens = {
+  Level.CONFIG: '\x1B[32m',
+  Level.INFO: '\x1B[34m',
+  Level.WARNING: '\x1B[93m',
+  Level.SEVERE: '\x1B[91m',
+  Level.SHOUT: '\x1B[95m',
+};
+
+class Logging {
+  /// Singleton instance of this class.
+  static Logging instance = Logging();
+
+  /// Whether logging has been initialized.
+  bool _bIsInitialized = false;
+
+  /// The sink to which log data should be written to send it to a file.
+  IOSink? _logFileSink;
+
+  /// A future that will return [_logFileSink] once it's created.
+  Future<IOSink>? _logFileSinkFuture;
+
+  /// The date when the log file was last updated.
+  DateTime? _lastLogChangeTime;
+
+  /// If true, it's time to change to a new log file.
+  bool get _bNeedsLogChange {
+    // We haven't started or finished creating a log
+    if (_logFileSinkFuture == null) {
+      return true;
+    }
+
+    // We have a file, so check that it isn't an outdated path
+    if (_logFileSink != null) {
+      assert(_lastLogChangeTime != null);
+
+      final now = DateTime.now();
+      return _fileDateFormatter.format(now) != _fileDateFormatter.format(_lastLogChangeTime!);
+    }
+
+    return false;
+  }
+
+  /// Given a log file, get the date on which it was created, or null if it couldn't be determined.
+  static DateTime? getDateForLog(File logFile) {
+    final name = logFile.uri.pathSegments.last;
+    final String? match = _logDateFromName.firstMatch(name)?[1];
+
+    if (match != null) {
+      try {
+        return _fileDateFormatter.parse(match);
+      } catch (error) {
+        return null;
+      }
+    }
+
+    return null;
+  }
+
+  /// Given a log file, make a user-friendly string to refer to it by.
+  static String getNameForLog(File logFile) {
+    String title = logFile.uri.pathSegments.last;
+
+    final DateTime? logDate = Logging.getDateForLog(logFile);
+    if (logDate != null) {
+      title += ' (${_friendlyDateFormatter.format(logDate)})';
+    }
+
+    return title;
+  }
+
+  /// Initialize logging for the app.
+  void initialize() {
+    if (_bIsInitialized) {
+      _log.warning('Tried to initialize logging more than once');
+      return;
+    }
+
+    _bIsInitialized = true;
+
+    Logger.root.level = Level.ALL;
+    Logger.root.onRecord.listen(_onLogEvent);
+  }
+
+  /// Clean up logging for the app.
+  void cleanUp() {
+    _logFileSink?.close();
+
+    _bIsInitialized = false;
+  }
+
+  /// Get the directory in which to store logs.
+  Future<Directory> getLogDirectory() async {
+    final Directory documentRoot = await getApplicationSupportDirectory();
+
+    return Directory('${documentRoot.path}/logs');
+  }
+
+  /// Get a reference to the current sink to which log data should be written.
+  Future<IOSink> _getActiveLogFileSink() {
+    if (!_bNeedsLogChange) {
+      return _logFileSinkFuture!;
+    }
+
+    // Flush and close the old sink if it exists
+    if (_logFileSink != null) {
+      final IOSink previousSink = _logFileSink!;
+      previousSink.flush().then((value) => previousSink.close());
+    }
+
+    _logFileSinkFuture = getLogDirectory().then((Directory directory) async {
+      try {
+        await directory.create(recursive: true);
+      } catch (e) {
+        return Future.error(e);
+      }
+
+      final fileName = 'epic_stage_app_${_fileDateFormatter.format(DateTime.now())}.log';
+      final logFile = File('${directory.path}/$fileName');
+
+      _logFileSink = logFile.openWrite(mode: FileMode.append);
+
+      if (_logFileSink != null) {
+      	_lastLogChangeTime = DateTime.now();
+        return _logFileSink!;
+      }
+
+      return Future.error('Failed to open log "$fileName"');
+    });
+
+    return _logFileSinkFuture!;
+  }
+
+  /// Called when a log event is received.
+  void _onLogEvent(LogRecord record) {
+    String formatted =
+        '[${_timeFormatter.format(record.time)}] ${record.level.name}/${record.loggerName}: ${record.message}';
+    if (record.error != null) {
+      formatted += '\n${record.error}';
+    }
+
+    if (record.stackTrace != null) {
+      formatted += '\n${record.stackTrace}';
+    }
+
+    _logToConsole(formatted, record);
+    _logToFile(formatted, record);
+  }
+
+  /// Send a log to the debug console.
+  void _logToConsole(String formatted, LogRecord record) {
+    final resetToken = '\x1B[0m';
+    final String colorToken = _logColorTokens[record.level] ?? '';
+    print('$colorToken$formatted$resetToken');
+  }
+
+  /// Send a log to the log file.
+  void _logToFile(String formatted, LogRecord record) async {
+    final IOSink sink = await _getActiveLogFileSink();
+    sink.writeln(formatted);
+  }
+}
diff --git a/Engine/Extras/VirtualProduction/EpicStageApp/lib/utilities/math_utilities.dart b/Engine/Extras/VirtualProduction/EpicStageApp/lib/utilities/math_utilities.dart
new file mode 100644
--- /dev/null
+++ b/Engine/Extras/VirtualProduction/EpicStageApp/lib/utilities/math_utilities.dart
@@ -0,0 +1,39 @@
+// Copyright Epic Games, Inc. All Rights Reserved.
+
+import 'dart:math';
+import 'dart:ui';
+
+/// Given a [value], apply the exponential curve produced by the [exponent].
+/// If [bInverse] is true, apply the inverse of the exponent.
+double exponentiateValue(double value, double exponent, bool bInverse) {
+  assert(exponent > 0);
+
+  if (exponent == 1.0) {
+    return value;
+  }
+
+  if (bInverse) {
+    if (exponent == 0) {
+      throw Exception('Can\'t apply inverse of exponent $exponent}');
+    }
+    exponent = 1.0 / exponent;
+  }
+
+  return pow(value, exponent) as double;
+}
+
+/// Given a [position] relative to a circle's center, apply the exponential curve produced by the [exponent] to the
+/// position's distance from the center. If [bInverse] is true, apply the inverse of the exponent.
+Offset exponentiateCirclePosition(Offset position, double exponent, bool bInverse) {
+  if (exponent == 1.0) {
+    return position;
+  }
+
+  final double distance = position.distance;
+  if (distance <= 0) {
+    return Offset.zero;
+  }
+
+  final double scale = exponentiateValue(distance, exponent, bInverse) / distance;
+  return position.scale(scale, scale);
+}
diff --git a/Engine/Extras/VirtualProduction/EpicStageApp/lib/utilities/net_utilities.dart b/Engine/Extras/VirtualProduction/EpicStageApp/lib/utilities/net_utilities.dart
new file mode 100644
--- /dev/null
+++ b/Engine/Extras/VirtualProduction/EpicStageApp/lib/utilities/net_utilities.dart
@@ -0,0 +1,213 @@
+// Copyright Epic Games, Inc. All Rights Reserved.
+
+import 'dart:async';
+import 'dart:io';
+import 'dart:typed_data';
+
+import 'package:archive/archive_io.dart';
+import 'package:logging/logging.dart';
+import 'package:uuid/uuid.dart';
+
+import '../models/engine_connection.dart';
+
+const int _protocolVersion = 0; // The current protocol version for beacon messages
+const String _protocolIdentifier = 'ES@p';
+
+/// Data about a detected engine connection
+class ConnectionData {
+  ConnectionData({required this.uuid, required this.name, required this.websocketAddress, required this.websocketPort})
+      : lastSeenTime = DateTime.now();
+
+  final UuidValue uuid;
+  String name;
+  InternetAddress websocketAddress;
+  int websocketPort;
+  DateTime lastSeenTime;
+}
+
+/// Convert a Uint8List to a UUID.
+UuidValue uuidFromUint8List(Uint8List data) {
+  String guidString = Uuid.unparse(data);
+  return UuidValue(guidString, false);
+}
+
+/// Make the beacon message that will be multicast to find compatible engine instances.
+Uint8List makeBeaconMessage() {
+  return Uint8List.fromList(_protocolIdentifier.codeUnits + [_protocolVersion]);
+}
+
+/// Retrieve the connection data from the datagram received from the response to a beacon message.
+/// If the connection data couldn't be retrieved, this will return null.
+ConnectionData? getConnectionFromBeaconResponse(Datagram datagram) {
+  final InputStream stream = InputStream(datagram.data);
+  stream.byteOrder = LITTLE_ENDIAN;
+  ConnectionData? connection;
+
+  try {
+    stream.readByte(); // Protocol version; currently ignored
+    final UuidValue engineUuid = uuidFromUint8List(stream.readBytes(16).toUint8List());
+    final int port = stream.readUint32();
+
+    final int nameLength = stream.readUint32();
+    String name = stream.readString(size: nameLength, utf8: true);
+
+    if (name[name.length - 1] == '\x00') {
+      // Remove the null terminator
+      name = name.substring(0, name.length - 1);
+    }
+
+    // If the engine didn't provide a friendly name, just use the address and port.
+    if (name.isEmpty || (name.length == 1 && name[0] == '\x00')) {
+      name = '${datagram.address.address}:${port.toString()}';
+    }
+
+    connection = ConnectionData(uuid: engineUuid, name: name, websocketAddress: datagram.address, websocketPort: port);
+  } catch (error) {
+    // TODO: Set up error logging so we can do something here
+  }
+
+  return connection;
+}
+
+/// A helper class that automatically retries sending a message to the engine if it times out.
+class MessageTimeoutRetryHelper {
+  MessageTimeoutRetryHelper({
+    required this.connection,
+    required this.sendMessage,
+    this.onTimeout,
+    required this.timeoutDuration,
+    this.logger,
+    this.logDescription,
+    this.bLogOnlyAfterFailure = true,
+  }) : assert(!(logDescription == null && logger != null));
+
+  /// The connection to test before attempting to retry.
+  final EngineConnectionManager connection;
+
+  /// A function that sends the message to the engine.
+  final void Function() sendMessage;
+
+  /// A function that will be called each time this times out, passing the consecutive [timeoutCount].
+  final void Function(int timeoutCount)? onTimeout;
+
+  /// How long to wait before timing out.
+  final Duration timeoutDuration;
+
+  /// The log to which debug logs should be sent. If null, this won't log anything.
+  final Logger? logger;
+
+  /// A debug description of what this is trying to send. If [logger] is set, this must be too.
+  final String? logDescription;
+
+  /// If true, only send a log on success if this failed at least once before succeeding.
+  final bool bLogOnlyAfterFailure;
+
+  /// Timer use to determine if we've timed out.
+  Timer? _timer;
+
+  /// Whether this is currently attempting to send a message.
+  bool _bIsActive = false;
+
+  /// Whether our last attempt was successful.
+  bool _bWasSuccessful = false;
+
+  /// How many consecutive times this timed out since being started. This is not cleared until [start] is called.
+  int _timeoutCount = 0;
+
+  /// Whether this is currently attempting to send a message.
+  bool get bIsActive => _bIsActive;
+
+  /// How many consecutive times this timed out since being started. This is not cleared until [start] is called.
+  int get timeoutCount => _timeoutCount;
+
+  /// Start attempting to send the message.
+  void start() {
+    if (bIsActive) {
+      return;
+    }
+
+    _bIsActive = true;
+    _bWasSuccessful = false;
+    _timeoutCount = 0;
+    _internalSendMessage();
+  }
+
+  /// Give up trying to send the message.
+  void giveUp() {
+    if (!_cancel()) {
+      return;
+    }
+
+    logger?.info('Abandoned "$logDescription" after $timeoutCount failed attempt(s)');
+  }
+
+  /// Stop trying to send the message after succeeding.
+  void succeed() {
+    if (!_cancel()) {
+      return;
+    }
+
+    _bWasSuccessful = true;
+  }
+
+  /// Cancel the timer and return to inactive state.
+  bool _cancel() {
+    if (!bIsActive) {
+      return false;
+    }
+
+    _bIsActive = false;
+    _timer?.cancel();
+
+    return true;
+  }
+
+  /// Send the message and start the timer.
+  void _internalSendMessage() {
+    sendMessage();
+
+    _timer?.cancel();
+    _timer = Timer(timeoutDuration, _internalOnTimeout);
+  }
+
+  /// Called when this times out waiting for a response.
+  void _internalOnTimeout() async {
+    ++_timeoutCount;
+
+    logger?.info('Timed out waiting for "$logDescription" (attempt #$timeoutCount)');
+
+    onTimeout?.call(_timeoutCount);
+
+    if (!_bIsActive) {
+      _logCancellation();
+      return;
+    }
+
+    logger?.info('Confirming engine is receiving messages for "$logDescription"');
+
+    // Check that our messages are getting through before we attempt to send another request. This confirms that the
+    // engine actually failed to receive/respond that we're not just waiting for it (or our request) to make it through
+    // network traffic
+    await connection.testEngineResponse();
+
+    if (!_bIsActive) {
+      logger?.info('Engine OK for "$logDescription"');
+      _logCancellation();
+      return;
+    }
+
+    logger?.info('Engine OK, retrying for "$logDescription"');
+    _internalSendMessage();
+  }
+
+  /// If needed, log that this has succeeded/failed.
+  void _logCancellation() {
+    if (_bWasSuccessful) {
+      if (!bLogOnlyAfterFailure || timeoutCount > 0) {
+        logger?.info('Succeeded at "$logDescription" after $timeoutCount failed attempt(s)');
+      }
+    } else {
+      logger?.info('Abandoned "$logDescription" after $timeoutCount failed attempt(s)');
+    }
+  }
+}
diff --git a/Engine/Extras/VirtualProduction/EpicStageApp/lib/utilities/preferences_bundle.dart b/Engine/Extras/VirtualProduction/EpicStageApp/lib/utilities/preferences_bundle.dart
new file mode 100644
--- /dev/null
+++ b/Engine/Extras/VirtualProduction/EpicStageApp/lib/utilities/preferences_bundle.dart
@@ -0,0 +1,24 @@
+// Copyright Epic Games, Inc. All Rights Reserved.
+
+import 'package:flutter/cupertino.dart';
+import 'package:provider/provider.dart';
+import 'package:streaming_shared_preferences/streaming_shared_preferences.dart';
+
+import 'transient_preference.dart';
+
+/// A convenience class holding both types of centralized preference classes so they can be easily accessed throughout
+/// the app.
+class PreferencesBundle {
+  const PreferencesBundle(this.persistent, this.transient);
+
+  /// Preferences that will be serialized to disk and persisted across runs of the app.
+  final StreamingSharedPreferences persistent;
+
+  /// Preferences that will be held in memory and cleared between runs of the app./**/
+  final TransientSharedPreferences transient;
+
+  /// Get the [PreferencesBundle] provided in the given [context].
+  static PreferencesBundle of(BuildContext context) {
+    return Provider.of(context, listen: false);
+  }
+}
diff --git a/Engine/Extras/VirtualProduction/EpicStageApp/lib/utilities/preloadable_shader.dart b/Engine/Extras/VirtualProduction/EpicStageApp/lib/utilities/preloadable_shader.dart
new file mode 100644
--- /dev/null
+++ b/Engine/Extras/VirtualProduction/EpicStageApp/lib/utilities/preloadable_shader.dart
@@ -0,0 +1,85 @@
+// Copyright Epic Games, Inc. All Rights Reserved.
+
+import 'dart:ui';
+
+import 'package:flutter/foundation.dart';
+
+/// Wrapper for a shader that can be loaded ahead of time to prevent flickering at runtime.
+abstract class PreloadableShader {
+  PreloadableShader() {
+    _createShaderInstance();
+  }
+
+  /// The current shader instance
+  FragmentShader? _shader;
+
+  /// The current shader instance
+  FragmentShader? get shader => _shader;
+
+  /// The program, which should be loaded before an instance of this class is created.
+  PreloadableShaderProgram get program;
+
+  /// Reload the shader program and re-create the shader instance from it.
+  void reload() {
+    program.load().then((_) => _createShaderInstance());
+  }
+
+  /// Dispose of the underlying shader.
+  @mustCallSuper
+  void dispose() {
+    _disposeShaderInstance();
+  }
+
+  /// Initialize the shader's uniform values when it's been freshly created.
+  @protected
+  void initUniforms() {}
+
+  /// Create a shader instance from the loaded program and initialize it.
+  void _createShaderInstance() {
+    _disposeShaderInstance();
+
+    if (program.program == null) {
+      throw Exception('Tried to create shader ${this.runtimeType} but the program was not loaded');
+    }
+
+    _shader = program.program!.fragmentShader();
+    initUniforms();
+  }
+
+  /// Dispose of the current shader instance, if any.
+  void _disposeShaderInstance() {
+    if (_shader != null) {
+      _shader!.dispose();
+      _shader = null;
+    }
+  }
+}
+
+/// Wrapper for a shader program to consistently load/unload it.
+class PreloadableShaderProgram {
+  PreloadableShaderProgram(this.assetPath);
+
+  /// The path to the asset containing the shader code.
+  final String assetPath;
+
+  /// THe loaded program, or null if it hasn't been loaded.
+  FragmentProgram? _program = null;
+
+  /// The loaded program, or null if it hasn't been loaded.
+  FragmentProgram? get program => _program;
+
+  /// Load the program.
+  Future<FragmentProgram> load() {
+    return FragmentProgram.fromAsset(assetPath).then((loadedProgram) {
+      _program = loadedProgram;
+      return loadedProgram;
+    });
+  }
+
+  /// Load the program.
+  void unload() {
+    if (program != null) {
+      _program = null;
+    }
+  }
+}
diff --git a/Engine/Extras/VirtualProduction/EpicStageApp/lib/utilities/streaming_preferences_extensions.dart b/Engine/Extras/VirtualProduction/EpicStageApp/lib/utilities/streaming_preferences_extensions.dart
new file mode 100644
--- /dev/null
+++ b/Engine/Extras/VirtualProduction/EpicStageApp/lib/utilities/streaming_preferences_extensions.dart
@@ -0,0 +1,155 @@
+// Copyright Epic Games, Inc. All Rights Reserved.
+
+import 'dart:convert';
+
+import 'package:logging/logging.dart';
+import 'package:shared_preferences/shared_preferences.dart';
+import 'package:streaming_shared_preferences/streaming_shared_preferences.dart';
+
+import '../models/settings/recent_actor_settings.dart';
+import '../models/unreal_object_filters.dart';
+import 'json_utilities.dart';
+
+final _log = Logger('PrefExtensions');
+
+/// Extensions to serialize/deserialize new types of preference values.
+extension PreferenceAdapterExtensions on StreamingSharedPreferences {
+  /// Starts with the current String set value for the given [key], then emits a new value every time there are changes
+  /// to the value associated with [key].
+  ///
+  /// If the value is null, starts with the value provided in [defaultValue]. When the value transitions from non-null
+  /// to null (ie. when the value is removed), emits [defaultValue].
+  Preference<Set<String>> getStringSet(
+    String key, {
+    required Set<String> defaultValue,
+  }) {
+    return getCustomValue(key, defaultValue: defaultValue, adapter: StringSetAdapter.instance);
+  }
+
+  /// Starts with the current Enum value for the given [key], then emits a new value every time there are changes to the
+  /// value associated with [key].
+  ///
+  /// If the value is null, starts with the value provided in [defaultValue]. When the value transitions from non-null
+  /// to null (ie. when the value is removed), emits [defaultValue].
+  Preference<T> getEnum<T extends Enum>(
+    String key, {
+    required T defaultValue,
+    required List<T> enumValues,
+  }) {
+    return getCustomValue(key, defaultValue: defaultValue, adapter: EnumAdapter(enumValues));
+  }
+
+  /// Starts with the current list value for the given [key], then emits a new value every time there are changes to the
+  /// value associated with [key].
+  ///
+  /// If the value is null, starts with the value provided in [defaultValue]. When the value transitions from non-null
+  /// to null (ie. when the value is removed), emits [defaultValue].
+  Preference<List<RecentlyPlacedActorData>> getRecentlyPlacedActorList(
+    String key, {
+    required List<RecentlyPlacedActorData> defaultValue,
+  }) {
+    return getCustomValue(key, defaultValue: defaultValue, adapter: RecentlyPlacedActorListAdapter.instance);
+  }
+
+  /// Starts with the current Set value for the given [key], then emits a new value every time there are changes to the
+  /// value associated with [key].
+  ///
+  /// If the value is null, starts with the value provided in [defaultValue]. When the value transitions from non-null
+  /// to null (ie. when the value is removed), emits [defaultValue].
+  Preference<Set<UnrealObjectFilter>> getUnrealObjectFilterSet(
+    String key, {
+    required Set<UnrealObjectFilter> defaultValue,
+  }) {
+    return getCustomValue(key, defaultValue: defaultValue, adapter: UnrealObjectFilterListAdapter.instance);
+  }
+}
+
+/// A [PreferenceAdapter] implementation for storing and retrieving a [Set] of [String] objects.
+class StringSetAdapter extends PreferenceAdapter<Set<String>> {
+  static const instance = StringSetAdapter._();
+  const StringSetAdapter._();
+
+  @override
+  Set<String>? getValue(preferences, key) => preferences.getStringList(key)?.toSet();
+
+  @override
+  Future<bool> setValue(preferences, key, values) => preferences.setStringList(key, values.toList(growable: false));
+}
+
+/// A [PreferenceAdapter] implementation for storing and retrieving an [Enum] value as a string.
+class EnumAdapter<T extends Enum> extends PreferenceAdapter<T> {
+  const EnumAdapter(this.enumValues);
+
+  final List<T> enumValues;
+
+  @override
+  T? getValue(preferences, key) {
+    try {
+      final String? name = preferences.getString(key);
+
+      if (name == null) {
+        return null;
+      }
+
+      return jsonToEnumValue(name, enumValues);
+    } catch (_) {
+      return null;
+    }
+  }
+
+  @override
+  Future<bool> setValue(preferences, key, value) => preferences.setString(key, enumToJsonValue(value));
+}
+
+/// A [PreferenceAdapter] implementation for storing and retrieving a [List] of [RecentlyPlacedActorData] values.
+class RecentlyPlacedActorListAdapter extends PreferenceAdapter<List<RecentlyPlacedActorData>> {
+  static const instance = RecentlyPlacedActorListAdapter._();
+  const RecentlyPlacedActorListAdapter._();
+
+  @override
+  List<RecentlyPlacedActorData>? getValue(SharedPreferences preferences, String key) {
+    try {
+      return preferences
+          .getStringList(key)
+          ?.map((jsonData) => RecentlyPlacedActorData.fromJson(json.decode(jsonData)))
+          .toList();
+    } catch (e) {
+      _log.warning('Failed to deserialize RecentlyPlacedActorData list: $e');
+      return null;
+    }
+  }
+
+  @override
+  Future<bool> setValue(SharedPreferences preferences, String key, List<RecentlyPlacedActorData> value) {
+    final encodedData = value.map((data) => json.encode(data.toJson())).toList(growable: false);
+    return preferences.setStringList(key, encodedData);
+  }
+}
+
+/// A [PreferenceAdapter] implementation for storing and retrieving a [Set] of [UnrealObjectFilter] values.
+class UnrealObjectFilterListAdapter extends PreferenceAdapter<Set<UnrealObjectFilter>> {
+  static const instance = UnrealObjectFilterListAdapter._();
+  const UnrealObjectFilterListAdapter._();
+
+  @override
+  Set<UnrealObjectFilter>? getValue(SharedPreferences preferences, String key) {
+    try {
+      return preferences
+          .getStringList(key)
+          ?.map((filterName) => UnrealObjectFilterRegistry.getFilterByName(filterName))
+          // Remove any filters that no longer exist
+          .where((element) => element != null)
+          .map((element) => element!)
+          .toSet();
+    } catch (e) {
+      _log.warning('Failed to deserialize RecentlyPlacedActorData list: $e');
+      return null;
+    }
+  }
+
+  @override
+  Future<bool> setValue(SharedPreferences preferences, String key, Set<UnrealObjectFilter> value) {
+    final nameList = value.map((filter) => filter.internalName).toList(growable: false);
+    return preferences.setStringList(key, nameList);
+  }
+}
diff --git a/Engine/Extras/VirtualProduction/EpicStageApp/lib/utilities/transient_preference.dart b/Engine/Extras/VirtualProduction/EpicStageApp/lib/utilities/transient_preference.dart
new file mode 100644
--- /dev/null
+++ b/Engine/Extras/VirtualProduction/EpicStageApp/lib/utilities/transient_preference.dart
@@ -0,0 +1,159 @@
+// Copyright Epic Games, Inc. All Rights Reserved.
+import 'dart:async';
+
+import 'package:flutter/material.dart';
+
+/// Maps transient preferences to names so they can be accessed from one shared place.
+class TransientSharedPreferences {
+  final Map<String, TransientPreference> _preferences = {};
+
+  /// Get the preference with the given [name].
+  TransientPreference<T> get<T>(String name, {required T defaultValue}) {
+    TransientPreference? preference = _preferences[name];
+
+    if (preference != null) {
+      if (!(preference is TransientPreference<T>)) {
+        throw Exception('Tried to get transient preference "$name" of type $T,'
+            'but stream already exists as ${preference.runtimeType}');
+      }
+
+      return preference;
+    }
+
+    preference = TransientPreference<T>._(defaultValue);
+    _preferences[name] = preference;
+    return preference as TransientPreference<T>;
+  }
+}
+
+// Code below based largely on the streaming_shared_preferences package.
+// Original license:
+//
+// Copyright 2019 Iiro Krankka
+//
+// Redistribution and use in source and binary forms, with or without modification, are permitted provided that the following conditions are met:
+//
+// 1. Redistributions of source code must retain the above copyright notice, this list of conditions and the following disclaimer.
+//
+// 2. Redistributions in binary form must reproduce the above copyright notice, this list of conditions and the following disclaimer in the documentation and/or other materials provided with the distribution.
+//
+// THIS SOFTWARE IS PROVIDED BY THE COPYRIGHT HOLDERS AND CONTRIBUTORS "AS IS" AND ANY EXPRESS OR IMPLIED WARRANTIES, INCLUDING, BUT NOT LIMITED TO, THE IMPLIED WARRANTIES OF MERCHANTABILITY AND FITNESS FOR A PARTICULAR PURPOSE ARE DISCLAIMED. IN NO EVENT SHALL THE COPYRIGHT HOLDER OR CONTRIBUTORS BE LIABLE FOR ANY DIRECT, INDIRECT, INCIDENTAL, SPECIAL, EXEMPLARY, OR CONSEQUENTIAL DAMAGES (INCLUDING, BUT NOT LIMITED TO, PROCUREMENT OF SUBSTITUTE GOODS OR SERVICES; LOSS OF USE, DATA, OR PROFITS; OR BUSINESS INTERRUPTION) HOWEVER CAUSED AND ON ANY THEORY OF LIABILITY, WHETHER IN CONTRACT, STRICT LIABILITY, OR TORT (INCLUDING NEGLIGENCE OR OTHERWISE) ARISING IN ANY WAY OUT OF THE USE OF THIS SOFTWARE, EVEN IF ADVISED OF THE POSSIBILITY OF SUCH DAMAGE.
+
+/// A single user value which does not persist outside the app. Exposed as a stream which emits a value whenever the
+/// preference's value changes.
+class TransientPreference<T> extends Stream<T> {
+  TransientPreference._(T defaultValue) : _currentValue = defaultValue;
+
+  final StreamController<T> _changeStream = StreamController<T>.broadcast(sync: true);
+
+  /// The current value of the preference.
+  T _currentValue;
+
+  /// Set the value of the preference and notify any listeners.
+  void setValue(T value) {
+    _currentValue = value;
+
+    if (_changeStream.hasListener) {
+      _changeStream.add(value);
+    }
+  }
+
+  /// Get the latest value of the preference.
+  T getValue() {
+    return _currentValue;
+  }
+
+  @override
+  StreamSubscription<T> listen(void Function(T event)? onData,
+      {Function? onError, void Function()? onDone, bool? cancelOnError}) {
+    return _changeStream.stream.listen(
+      onData,
+      onError: onError,
+      onDone: onDone,
+      cancelOnError: cancelOnError,
+    );
+  }
+}
+
+/// A function that builds a widget whenever a [Preference] has a new value.
+typedef TransientPreferenceWidgetBuilder<T> = Function(BuildContext context, T value);
+
+/// A builder which gets a value from a [TransientPreference] and updates whenever it changes.
+class TransientPreferenceBuilder<T> extends StatefulWidget {
+  TransientPreferenceBuilder({
+    required this.preference,
+    required this.builder,
+  });
+
+  /// The preference on which you want to react and rebuild your widgets based on.
+  final TransientPreference<T> preference;
+
+  /// The function that builds a widget when a [preference] has new data.
+  final TransientPreferenceWidgetBuilder<T> builder;
+
+  @override
+  _TransientPreferenceBuilderState<T> createState() => _TransientPreferenceBuilderState<T>();
+}
+
+class _TransientPreferenceBuilderState<T> extends State<TransientPreferenceBuilder<T>> {
+  late final T _initialData;
+  late final Stream<T> _preference;
+
+  @override
+  void initState() {
+    super.initState();
+    _initialData = widget.preference.getValue();
+    _preference = widget.preference.transform(_EmitOnlyChangedValues(_initialData));
+  }
+
+  @override
+  Widget build(BuildContext context) {
+    return StreamBuilder<T>(
+      initialData: _initialData,
+      stream: _preference,
+      builder: (context, snapshot) => widget.builder(context, snapshot.data!),
+    );
+  }
+}
+
+/// Makes sure that [PreferenceBuilder] does not run its builder function if the
+/// new value is identical to the last one.
+class _EmitOnlyChangedValues<T> extends StreamTransformerBase<T, T> {
+  _EmitOnlyChangedValues(this.startValue);
+  final T startValue;
+
+  @override
+  Stream<T> bind(Stream<T> stream) {
+    return StreamTransformer<T, T>((input, cancelOnError) {
+      T? lastValue = startValue;
+
+      late final StreamController<T> controller;
+      late final StreamSubscription<T> subscription;
+
+      controller = StreamController<T>(
+        sync: true,
+        onListen: () {
+          subscription = input.listen(
+            (value) {
+              if (value != lastValue) {
+                controller.add(value);
+                lastValue = value;
+              }
+            },
+            onError: controller.addError,
+            onDone: controller.close,
+            cancelOnError: cancelOnError,
+          );
+        },
+        onPause: ([resumeSignal]) => subscription.pause(resumeSignal),
+        onResume: () => subscription.resume(),
+        onCancel: () {
+          lastValue = null;
+          return subscription.cancel();
+        },
+      );
+
+      return controller.stream.listen(null);
+    }).bind(stream);
+  }
+}
diff --git a/Engine/Extras/VirtualProduction/EpicStageApp/lib/utilities/unreal_colors.dart b/Engine/Extras/VirtualProduction/EpicStageApp/lib/utilities/unreal_colors.dart
new file mode 100644
--- /dev/null
+++ b/Engine/Extras/VirtualProduction/EpicStageApp/lib/utilities/unreal_colors.dart
@@ -0,0 +1,29 @@
+// Copyright Epic Games, Inc. All Rights Reserved.
+
+import 'dart:ui';
+
+/// A static class containing colors used in Unreal's editor.
+class UnrealColors {
+  static const Color highlightBlue = Color(0xff0070e0);
+  static const Color highlightRed = Color(0xffef3535);
+  static const Color highlightGreen = Color(0xff8bc24a);
+  static const Color highlightTeal = Color(0xff00a78f);
+  static const Color highlightPurple = Color(0xff8637ac);
+  static const Color warning = Color(0xffffa640);
+  static const Color black = Color(0xff000000);
+  static const Color white = Color(0xffffffff);
+
+  // Gray colors are named after the % lightness (as in HSL's L component)
+  static const Color gray06 = Color(0xff0f0f0f);
+  static const Color gray10 = Color(0xff1a1a1a);
+  static const Color gray13 = Color(0xff202020);
+  static const Color gray14 = Color(0xff242424);
+  static const Color gray18 = Color(0xff2f2f2f);
+  static const Color gray22 = Color(0xff383838);
+  static const Color gray31 = Color(0xff505050);
+  static const Color gray42 = Color(0xff6a6a6a);
+  static const Color gray56 = Color(0xff909090);
+  static const Color gray75 = Color(0xffc0c0c0);
+  static const Color gray90 = Color(0xffe6e6e6);
+  static const Color gray94 = Color(0xfff0f0f0);
+}
diff --git a/Engine/Extras/VirtualProduction/EpicStageApp/lib/utilities/unreal_utilities.dart b/Engine/Extras/VirtualProduction/EpicStageApp/lib/utilities/unreal_utilities.dart
new file mode 100644
--- /dev/null
+++ b/Engine/Extras/VirtualProduction/EpicStageApp/lib/utilities/unreal_utilities.dart
@@ -0,0 +1,233 @@
+// Copyright Epic Games, Inc. All Rights Reserved.
+
+import 'dart:math' as math;
+import 'dart:ui';
+
+import 'package:logging/logging.dart';
+import 'package:vector_math/vector_math_64.dart' as vec;
+
+import '../models/color.dart';
+import '../models/unreal_types.dart';
+
+/// Type name to use to represent all Unreal enums, regardless of their actual type name.
+String unrealEnumTypeName = '_Enum';
+
+final _log = Logger('UnrealUtilities');
+
+/// Data about how to convert an Unreal type into/out of JSON.
+class _UnrealTypeConversionData {
+  const _UnrealTypeConversionData({required this.jsonToDart, required this.dartToJson});
+
+  /// Function that converts the type's JSON data ([jsonValue]) to its native Dart equivalent.
+  /// [previousValue] is the previous value before [jsonValue] was received, or null if no previous value.
+  /// [conversionMetadata] contains arbitrary key/value data which may be used in some cases to disambiguate how the
+  /// data should be converted.
+  final dynamic Function(dynamic jsonValue, dynamic previousValue, Map<String, dynamic> conversionMetadata) jsonToDart;
+
+  /// Function that converts the type's native Dart type to its JSON equivalent.
+  /// [conversionMetadata] contains arbitrary key/value data which may be used in some cases to disambiguate how the
+  /// data should be converted.
+  final dynamic Function(dynamic dartValue, Map<String, dynamic> conversionMetadata) dartToJson;
+}
+
+_UnrealTypeConversionData _noConversion = _UnrealTypeConversionData(
+  jsonToDart: (dynamic jsonValue, dynamic previousValue, Map<String, dynamic> conversionMetadata) => jsonValue,
+  dartToJson: (dynamic dartValue, Map<String, dynamic> conversionMetadata) => dartValue,
+);
+
+/// Map from Unreal type name to conversion data for the associated class.
+Map<String, _UnrealTypeConversionData> _typeConverters = {
+  'bool': _noConversion,
+  'uint8': _UnrealTypeConversionData(
+    jsonToDart: (dynamic jsonValue, dynamic previousValue, Map<String, dynamic> conversionMetadata) {
+      if (jsonValue is bool) {
+        // Boolean values are often represented as a single uint8 bit
+        return jsonValue;
+      }
+
+      return jsonNumberToInt(jsonValue);
+    },
+    dartToJson: _noConversion.dartToJson,
+  ),
+  'float': _UnrealTypeConversionData(
+    jsonToDart: (dynamic jsonValue, dynamic previousValue, Map<String, dynamic> conversionMetadata) =>
+        jsonNumberToDouble(jsonValue, null),
+    dartToJson: _noConversion.dartToJson,
+  ),
+  'double': _UnrealTypeConversionData(
+    jsonToDart: (dynamic jsonValue, dynamic previousValue, Map<String, dynamic> conversionMetadata) =>
+        jsonNumberToDouble(jsonValue, null),
+    dartToJson: _noConversion.dartToJson,
+  ),
+  'FLinearColor': _UnrealTypeConversionData(
+    jsonToDart: (dynamic jsonValue, dynamic previousValue, Map<String, dynamic> conversionMetadata) {
+      final floatColor = FloatColor(
+        jsonNumberToDouble(jsonValue['R'], 0.0)!,
+        jsonNumberToDouble(jsonValue['G'], 0.0)!,
+        jsonNumberToDouble(jsonValue['B'], 0.0)!,
+        jsonNumberToDouble(jsonValue['A'], 0.0)!,
+      );
+
+      final hsvColor = NegatableHSVColor.fromColor(floatColor);
+
+      const minValueForPosition = 0.0001;
+      if (previousValue != null && hsvColor.value <= minValueForPosition) {
+        // The RGBA color's value is too low to determine the position on the color wheel. Instead, we'll use the
+        // last known position on our end.
+        return WheelColor(previousValue.position, hsvColor.value, hsvColor.alpha);
+      }
+
+      return WheelColor.fromHSVColor(hsvColor);
+    },
+    dartToJson: (dynamic dartValue, Map<String, dynamic> conversionMetadata) => dartValue.toJson(),
+  ),
+  'FVector2D': _UnrealTypeConversionData(
+    jsonToDart: (dynamic jsonValue, dynamic previousValue, Map<String, dynamic> conversionMetadata) => vec.Vector2(
+      jsonNumberToDouble(jsonValue['X'], 0.0)!,
+      jsonNumberToDouble(jsonValue['Y'], 0.0)!,
+    ),
+    dartToJson: (dynamic dartValue, Map<String, dynamic> conversionMetadata) => {'X': dartValue.x, 'Y': dartValue.y},
+  ),
+  'FVector4_Color': _UnrealTypeConversionData(
+    jsonToDart: (dynamic jsonValue, dynamic previousValue, Map<String, dynamic> conversionMetadata) {
+      final floatColor = FloatColor(
+        jsonNumberToDouble(jsonValue['X'], 0.0)!,
+        jsonNumberToDouble(jsonValue['Y'], 0.0)!,
+        jsonNumberToDouble(jsonValue['Z'], 0.0)!,
+        jsonNumberToDouble(jsonValue['W'], 0.0)!,
+      );
+
+      final double minValue = conversionMetadata['minValue'] ?? 0.0;
+      final double maxValue = conversionMetadata['maxValue'] ?? 1.0;
+      final double valueRange = maxValue - minValue;
+      final double actualValue = math.max(floatColor.red, math.max(floatColor.green, floatColor.blue));
+
+      assert(valueRange != 0);
+
+      // Convert to a [0, 1] range first in order to produce a valid HSV color, but preserve the original value and
+      // alpha (luminance) components.
+      final FloatColor scaledColor =
+          ((floatColor - FloatColor(minValue, minValue, minValue, minValue)) / valueRange).withAlpha(floatColor.alpha);
+      final hsvColor = NegatableHSVColor.fromColor(scaledColor).withValue(actualValue);
+
+      const minValueEpsilonForPosition = 0.0001;
+      if (previousValue != null && actualValue <= minValue + minValueEpsilonForPosition) {
+        // The RGBA color's value is too small to determine the position on the color wheel. Instead, we'll use the
+        // last known position on our end.
+        return WheelColor(previousValue.position, hsvColor.value, floatColor.alpha);
+      }
+
+      return WheelColor.fromHSVColor(hsvColor);
+    },
+    dartToJson: (dynamic dartValue, Map<String, dynamic> conversionMetadata) => dartValue.toJson(),
+  ),
+  unrealEnumTypeName: _noConversion,
+};
+
+/// Check whether the given Unreal [typeName] can be converted to a native Dart type.
+bool canConvertUnrealType(String typeName) {
+  return _typeConverters.containsKey(typeName);
+}
+
+/// Convert JSON [data] received from Unreal to the corresponding native Dart type based its [typeName].
+/// You may provide the [previousValue] as a hint. For example when converting RGBA data to a color, we use the
+/// previous value to fill in ambiguous hue/saturation values when the value is too low to calculate them unambiguously.
+/// Some conversion functions may also make make use of arbitrary data stored in [conversionMetadata].
+dynamic convertUnrealTypeJsonToDart(String typeName, dynamic data,
+    {dynamic previousValue, Map<String, dynamic>? conversionMetadata}) {
+  if (data == null) {
+    return null;
+  }
+
+  return _findConverter(typeName)?.jsonToDart(data, previousValue, conversionMetadata ?? {});
+}
+
+/// Convert native Dart [data] to the expected JSON format for Unreal based its [typeName].
+dynamic convertUnrealTypeDartToJson(String typeName, dynamic data, {Map<String, dynamic>? conversionMetadata}) {
+  if (data == null) {
+    return null;
+  }
+
+  return _findConverter(typeName)?.dartToJson(data, conversionMetadata ?? {});
+}
+
+/// Create a message in standard Unreal WebSocket format.
+dynamic createUnrealWebSocketMessage(String messageName, Object parameters) {
+  return {'MessageName': messageName, 'Parameters': parameters};
+}
+
+/// Find the converter for the given Unreal type name, or log the failure and return null if it doesn't exist.
+_UnrealTypeConversionData? _findConverter(String typeName) {
+  final _UnrealTypeConversionData? conversionData = _typeConverters[typeName];
+
+  if (conversionData == null) {
+    _log.warning('No converter exists for Unreal type "$typeName"', null, StackTrace.current);
+    return null;
+  }
+
+  return conversionData;
+}
+
+/// Convert a number received in JSON format (which could be a string, int, double, or null) to a double.
+/// If the number is null or can't be determined, defaultValue will be returned.
+double? jsonNumberToDouble(dynamic number, [double? defaultValue = 0]) {
+  if (number == null) {
+    return defaultValue;
+  }
+
+  if (number is String) {
+    return double.tryParse(number) ?? defaultValue;
+  }
+
+  try {
+    return number.toDouble();
+  } catch (e) {
+    return defaultValue;
+  }
+}
+
+/// Convert a number received in JSON format (which could be a string, int, double, or null) to an int.
+/// If the number is null or can't be determined, defaultValue will be returned.
+int? jsonNumberToInt(dynamic number, [int? defaultValue = 0]) {
+  if (number == null) {
+    return defaultValue;
+  }
+
+  if (number is String) {
+    return int.tryParse(number) ?? defaultValue;
+  }
+
+  try {
+    return number.toInt();
+  } catch (e) {
+    return defaultValue;
+  }
+}
+
+/// Convert an offset to the JSON format Unreal expects for a 2D vector.
+dynamic offsetToJson(Offset offset) {
+  return {
+    'X': offset.dx,
+    'Y': offset.dy,
+  };
+}
+
+/// Given a list of [properties], create a list of properties with [suffix] appended to their names.
+List<UnrealProperty> getPropertiesWithSuffix(List<UnrealProperty> properties, String suffix,
+    {String? typeNameOverride}) {
+  return properties
+      .map(
+        (UnrealProperty property) => UnrealProperty(
+          objectPath: property.objectPath,
+          propertyName: property.propertyName + suffix,
+          typeNameOverride: typeNameOverride,
+        ),
+      )
+      .toList();
+}
+
+/// Given a list of [properties], create a list of properties residing within those properties.
+List<UnrealProperty> getSubproperties(List<UnrealProperty> properties, String subpropertyName,
+    {String? typeNameOverride}) {
+  return getPropertiesWithSuffix(properties, '.$subpropertyName', typeNameOverride: typeNameOverride);
+}
diff --git a/Engine/Extras/VirtualProduction/EpicStageApp/lib/widgets/elements/asset_icon.dart b/Engine/Extras/VirtualProduction/EpicStageApp/lib/widgets/elements/asset_icon.dart
new file mode 100644
--- /dev/null
+++ b/Engine/Extras/VirtualProduction/EpicStageApp/lib/widgets/elements/asset_icon.dart
@@ -0,0 +1,62 @@
+// Copyright Epic Games, Inc. All Rights Reserved.
+
+import 'package:flutter/material.dart';
+import 'package:flutter_svg/svg.dart';
+
+/// An icon from an image asset (either a PNG or an SVG).
+class AssetIcon extends StatelessWidget {
+  const AssetIcon({
+    Key? key,
+    required this.path,
+    this.size,
+    this.width,
+    this.height,
+    this.color,
+    this.fit,
+  })  : assert(
+          size == null || (width == null && height == null),
+          'size parameter overrides width and height parameters',
+        ),
+        super(key: key);
+
+  /// The path of the asset. Must be an image or an SVG.
+  final String path;
+
+  /// If provided, use this as both the [width] and the [height].
+  final double? size;
+
+  /// The width of the icon.
+  final double? width;
+
+  /// The height of the icon.
+  final double? height;
+
+  /// The color of the icon.
+  final Color? color;
+
+  /// How to inscribe the image into the space allocated during layout.
+  final BoxFit? fit;
+
+  @override
+  Widget build(BuildContext context) {
+    if (path.endsWith('.svg')) {
+      return SvgPicture.asset(
+        path,
+        width: size ?? width,
+        height: size ?? height,
+        color: color,
+        colorBlendMode: BlendMode.modulate,
+        fit: fit ?? BoxFit.contain,
+      );
+    } else {
+      return Image.asset(
+        path,
+        width: size ?? width,
+        height: size ?? height,
+        color: color,
+        colorBlendMode: BlendMode.modulate,
+        fit: fit,
+      );
+    }
+  }
+}
diff --git a/Engine/Extras/VirtualProduction/EpicStageApp/lib/widgets/elements/color_grading_wheel.dart b/Engine/Extras/VirtualProduction/EpicStageApp/lib/widgets/elements/color_grading_wheel.dart
new file mode 100644
--- /dev/null
+++ b/Engine/Extras/VirtualProduction/EpicStageApp/lib/widgets/elements/color_grading_wheel.dart
@@ -0,0 +1,917 @@
+// Copyright Epic Games, Inc. All Rights Reserved.
+
+import 'dart:async';
+import 'dart:math' as math;
+
+import 'package:async/async.dart';
+import 'package:flutter/material.dart';
+import 'package:flutter_gen/gen_l10n/app_localizations.dart';
+import 'package:provider/provider.dart';
+
+import '../../../models/color.dart';
+import '../../../models/property_modify_operations.dart';
+import '../../../utilities/math_utilities.dart';
+import '../../models/settings/color_wheel_settings.dart';
+import '../../models/settings/delta_widget_settings.dart';
+import '../../utilities/constants.dart';
+import '../../utilities/guarded_refresh_state.dart';
+import '../../utilities/preferences_bundle.dart';
+import '../../utilities/transient_preference.dart';
+import '../../utilities/unreal_colors.dart';
+import '../../utilities/unreal_utilities.dart';
+import 'asset_icon.dart';
+import 'color_slider.dart';
+import 'delta_slider.dart';
+import 'delta_widget_base.dart';
+import 'selector_bar.dart';
+import 'unreal_widget_base.dart';
+
+/// Diameter of a dot on the color wheel indicating a value.
+const double _dotDiameter = 22.0;
+
+/// Fraction of the available horizontal space that the wheel should attempt to take up with its diameter.
+const double _wheelWidthFraction = 0.75;
+
+/// Maximum diameter of the color wheel regardless of available space.
+const double _maxWheelDiameter = 320;
+
+/// Minimum screen height to display all the sliders under the color grading wheel.
+const double _minScreenHeightForFullSize = 800;
+
+typedef ColorWheelValueData = DeltaWidgetValueData<WheelColor>;
+
+/// How to display the alpha value of a color.
+enum AlphaDisplayMode {
+  /// Don't let the user modify the alpha value at all.
+  hide,
+
+  /// Show the alpha value as an alpha slider.
+  alpha,
+
+  /// Show the alpha value as a luminance slider.
+  luminance,
+}
+
+/// Configuration data indicating how to display and control a color grading subproperty.
+class ColorGradingSubpropertyConfig {
+  const ColorGradingSubpropertyConfig({
+    required this.name,
+    required this.getDisplayName,
+    this.minValue = 0.0,
+    this.maxValue = 2.0,
+    this.saturationExponent = 2.4,
+  });
+
+  /// The name of the subproperty to control.
+  final String name;
+
+  /// A function returning the localized display name of the property.
+  final String Function(BuildContext context) getDisplayName;
+
+  /// The minimum value of the color's components.
+  final double minValue;
+
+  /// The maximum value of the color's components.
+  final double maxValue;
+
+  /// Exponent applied to the saturation component of the wheel, making it less sensitive towards the center.
+  final double saturationExponent;
+}
+
+/// A delta-based color wheel widget that controls a LinearColor property remotely in Unreal Engine.
+class UnrealColorWheel extends UnrealWidget {
+  const UnrealColorWheel({
+    super.key,
+    super.overrideName,
+    required super.unrealProperties,
+    this.extraTopWidget,
+    this.extraBottomWidget,
+    this.bShowAlpha = true,
+    this.bShowValueWithWheel = false,
+    this.maxRadius,
+  });
+
+  /// An extra widget to show next to the RGB/HSV and wheel/slider toggles.
+  final Widget? extraTopWidget;
+
+  /// An extra widget to show under the wheel, but above the value slider.
+  final Widget? extraBottomWidget;
+
+  /// If false, hide the alpha slider.
+  final bool bShowAlpha;
+
+  /// If false, hide the value slider when in wheel mode.
+  final bool bShowValueWithWheel;
+
+  /// If provided, don't let the wheel grow larger than this.
+  final double? maxRadius;
+
+  @override
+  _UnrealColorWheelState createState() => _UnrealColorWheelState();
+}
+
+class _UnrealColorWheelState extends State<UnrealColorWheel>
+    with UnrealWidgetStateMixin<UnrealColorWheel, WheelColor>, _UnrealRgbaComponentSliderCreator<UnrealColorWheel> {
+  @override
+  List<String> get rgbaComponentNames => const ['R', 'G', 'B', 'A'];
+
+  @override
+  double? get minComponentValue => null;
+
+  @override
+  double? get maxComponentValue => null;
+
+  @override
+  PropertyModifyOperation get modifyOperation => const WheelColorAddOperation();
+
+  @override
+  Widget build(BuildContext context) {
+    return TransientPreferenceBuilder(
+      preference: Provider.of<DeltaWidgetSettings>(context, listen: false).bIsInResetMode,
+      builder: (context, final bool bIsInResetMode) => DrivenColorWheel(
+        wheelValues: makeValueDataList(),
+        onValueChanged: handleOnChangedByUser,
+        onReset: handleOnResetByUser,
+        onInteractionFinished: endTransaction,
+        extraHeaderWidget: widget.extraTopWidget,
+        extraTopWidget: widget.extraBottomWidget,
+        alphaMode: widget.bShowAlpha ? AlphaDisplayMode.alpha : AlphaDisplayMode.hide,
+        createRgbaSlider: createComponentSlider,
+        maxRadius: widget.maxRadius,
+        bIsEnabled: !bIsInResetMode,
+      ),
+    );
+  }
+}
+
+/// A delta-based color grading wheel widget that controls a Vector4 color property remotely in Unreal Engine.
+class UnrealColorGradingWheel extends UnrealWidget {
+  const UnrealColorGradingWheel({
+    super.key,
+    super.overrideName,
+    super.enableProperties,
+    required super.unrealProperties,
+    this.extraTopWidget,
+    this.extraBottomWidget,
+    this.saturationExponent = 2.4,
+    this.minValue = 0.0,
+    this.maxValue = 2.0,
+    this.bShowLuminance = true,
+    this.bShowValueWithWheel = false,
+    this.maxRadius,
+  });
+
+  /// An extra widget to show next to the RGB/HSV and wheel/slider toggles.
+  final Widget? extraTopWidget;
+
+  /// An extra widget to show under the wheel, but above the value slider.
+  final Widget? extraBottomWidget;
+
+  /// Exponent applied to the saturation component of the wheel, making it less sensitive towards the center.
+  final double saturationExponent;
+
+  /// The minimum value of the color's components. The value will be clamped in this range and scaled based on the
+  /// min/max range when received in the engine.
+  final double minValue;
+
+  /// The maximum value of the color's components. The value will be clamped in this range and scaled based on the
+  /// min/max range when received in the engine.
+  final double maxValue;
+
+  /// If false, hide the value slider.
+  final bool bShowLuminance;
+
+  /// If false, hide the value slider when in wheel mode.
+  final bool bShowValueWithWheel;
+
+  /// If provided, don't let the wheel grow larger than this.
+  final double? maxRadius;
+
+  @override
+  _UnrealColorGradingWheelState createState() => _UnrealColorGradingWheelState();
+}
+
+class _UnrealColorGradingWheelState extends State<UnrealColorGradingWheel>
+    with
+        UnrealWidgetStateMixin<UnrealColorGradingWheel, WheelColor>,
+        _UnrealRgbaComponentSliderCreator<UnrealColorGradingWheel> {
+  @override
+  List<String> get rgbaComponentNames => const ['X', 'Y', 'Z', 'W'];
+
+  @override
+  double? get minComponentValue => widget.minValue;
+
+  @override
+  double? get maxComponentValue => widget.maxValue;
+
+  @override
+  PropertyModifyOperation get modifyOperation => WheelColorGradingAddOperation(
+        saturationExponent: widget.saturationExponent,
+        minValue: widget.minValue,
+        maxValue: widget.maxValue,
+      );
+
+  @override
+  Map<String, dynamic>? get conversionMetadata => {
+        'minValue': widget.minValue,
+        'maxValue': widget.maxValue,
+      };
+
+  @override
+  Widget build(BuildContext context) {
+    return TransientPreferenceBuilder(
+      preference: Provider.of<DeltaWidgetSettings>(context, listen: false).bIsInResetMode,
+      builder: (context, final bool bIsInResetMode) => DrivenColorWheel(
+        wheelValues: makeValueDataList(),
+        onValueChanged: handleOnChangedByUser,
+        onReset: handleOnResetByUser,
+        onInteractionFinished: endTransaction,
+        extraHeaderWidget: widget.extraTopWidget,
+        extraTopWidget: widget.extraBottomWidget,
+        saturationExponent: widget.saturationExponent,
+        minValue: widget.minValue,
+        maxValue: widget.maxValue,
+        alphaMode: widget.bShowLuminance ? AlphaDisplayMode.luminance : AlphaDisplayMode.hide,
+        createRgbaSlider: createComponentSlider,
+        maxRadius: widget.maxRadius,
+        bIsEnabled: !bIsInResetMode,
+      ),
+    );
+  }
+}
+
+mixin _UnrealRgbaComponentSliderCreator<WidgetType extends UnrealWidget> on State<WidgetType> {
+  /// Fill colors for the sliders corresponding to each RGBA color component.
+  static const List<Color?> _rgbaComponentFillColors = [
+    Color(0xffff0000),
+    Color(0xff00ff00),
+    Color(0xff0000ff),
+    null,
+  ];
+
+  /// The names of the 4 subproperties used to access color components.
+  List<String> get rgbaComponentNames;
+
+  /// The minimum value of a component slider.
+  double? get minComponentValue;
+
+  /// The maximum value of a component slider.
+  double? get maxComponentValue;
+
+  /// Create a slider that directly controls color component values. This lets us avoid converting from a
+  /// Vector4/LinearColor to a WheelColor and back (which causes precision loss) and instead directly accessing the
+  /// color's components by tracking subproperties of the color properties.
+  Widget createComponentSlider(int componentIndex, String componentLabel) {
+    return UnrealDeltaSlider(
+      unrealProperties: getSubproperties(widget.unrealProperties, rgbaComponentNames[componentIndex]),
+      hardMin: minComponentValue,
+      hardMax: maxComponentValue,
+      bShowResetButton: false,
+      overrideName: componentLabel,
+      enableProperties: widget.enableProperties,
+      fillColor: _rgbaComponentFillColors[componentIndex],
+    );
+  }
+}
+
+/// A color wheel which supports editing multiple objects and uses delta-based modification.
+/// Its values are controlled from outside of the widget itself.
+class DrivenColorWheel extends StatefulWidget {
+  const DrivenColorWheel({
+    Key? key,
+    required this.wheelValues,
+    required this.onValueChanged,
+    this.extraHeaderWidget,
+    this.extraTopWidget,
+    this.onReset,
+    this.onInteractionFinished,
+    this.saturationExponent = 1.0,
+    this.minValue = 0.0,
+    this.maxValue = 1.0,
+    this.alphaMode = AlphaDisplayMode.alpha,
+    this.createRgbaSlider,
+    this.maxRadius,
+    this.bIsEnabled = true,
+  }) : super(key: key);
+
+  /// An extra widget to show next to the RGB/HSV and wheel/slider toggles in the header.
+  final Widget? extraHeaderWidget;
+
+  /// An extra widget to show under the header, but above the wheel.
+  final Widget? extraTopWidget;
+
+  /// List of wheel colors for each property this widget is controlling.
+  final List<ColorWheelValueData?> wheelValues;
+
+  /// Function called when the user changes the color values.
+  /// Passes the amount by which each value changed. Note that color values may be negative to indicate a decrease.
+  final void Function(List<WheelColor>) onValueChanged;
+
+  /// Function called when the user is done interacting with the widget.
+  final void Function()? onInteractionFinished;
+
+  /// Function called when the reset button is pressed. Called after the value has been reset.
+  final void Function()? onReset;
+
+  /// Exponent applied to the saturation component of the wheel, making it less sensitive towards the center.
+  final double saturationExponent;
+
+  /// Minimum value of color's value range.
+  final double minValue;
+
+  /// Maximum value of color's value range.
+  final double maxValue;
+
+  /// How to display the color's alpha value to the user.
+  final AlphaDisplayMode alphaMode;
+
+  /// An optional function used to create the RGBA component slider widgets. If provided, changes to the sliders will
+  /// not result in [onValueChanged] callbacks, so they are responsible for reporting any changes.
+  /// [componentIndex] is an index into [rgbaSubpropertyNames] and [componentLabel] is the label to display on the
+  /// widget.
+  final Widget Function(int componentIndex, String componentLabel)? createRgbaSlider;
+
+  /// If provided, don't let the wheel grow larger than this.
+  final double? maxRadius;
+
+  /// If false, grey out the sliders and don't accept inputs. Note that this doesn't affect the reset button.
+  final bool bIsEnabled;
+
+  @override
+  _DrivenColorWheelState createState() => _DrivenColorWheelState();
+}
+
+class _DrivenColorWheelState extends State<DrivenColorWheel> with DeltaWidgetStateMixin, GuardedRefreshState {
+  late final ColorWheelSettings _settings;
+  StreamSubscription? _settingsSubscription;
+
+  /// Labels to show for the first three components of an RGB value.
+  List<String> get _rgbComponentLabels {
+    final localizations = AppLocalizations.of(context)!;
+    return [
+      localizations.colorComponentRed,
+      localizations.colorComponentGreen,
+      localizations.colorComponentBlue,
+    ];
+  }
+
+  @override
+  void initState() {
+    super.initState();
+
+    _settings = ColorWheelSettings(PreferencesBundle.of(context));
+
+    // The entire widget needs to redraw when any of these settings change, so handle them as a single subscription
+    // instead of PreferenceBuilders.
+    _settingsSubscription = StreamGroup.merge([
+      _settings.colorMode,
+      _settings.editMode,
+    ]).listen(refreshOnData);
+  }
+
+  @override
+  void dispose() {
+    _settingsSubscription?.cancel();
+
+    super.dispose();
+  }
+
+  @override
+  Widget build(BuildContext context) {
+    return LayoutBuilder(builder: (BuildContext context, BoxConstraints constraints) {
+      final localizations = AppLocalizations.of(context)!;
+      final bIsInResetModePref = Provider.of<DeltaWidgetSettings>(context, listen: false).bIsInResetMode;
+
+      final editMode = _settings.editMode.getValue();
+      final colorMode = _settings.colorMode.getValue();
+
+      final bool isFullHeight = MediaQuery.of(context).size.height > _minScreenHeightForFullSize;
+
+      final List<Widget> sliders = []; // List of sliders to display
+      Widget? colorWheel = null; // The color wheel (if it will be displayed)
+
+      late final Widget mainBody; // The main content of the widget (either a wheel or color component sliders)
+      late final List<Widget> otherWidgets = []; // Any other widgets we need to display under the main body
+
+      final double wheelDiameter = _getWheelDiameter(constraints);
+
+      // Create slider widgets, used both in full-height or slider-only modes
+      if (editMode == ColorWheelEditMode.sliders || isFullHeight) {
+        switch (colorMode) {
+          case ColorWheelColorMode.rgb:
+            if (widget.createRgbaSlider != null) {
+              // Create custom widgets for the first 3 components (R/G/B)
+              for (int componentIndex = 0; componentIndex < 3; ++componentIndex) {
+                sliders.add(widget.createRgbaSlider!(componentIndex, _rgbComponentLabels[componentIndex]));
+              }
+            } else {
+              // Use sliders that convert to components from the wheel value
+              sliders.addAll([
+                DrivenColorRedSlider(
+                  key: Key('Red'),
+                  values: widget.wheelValues,
+                  onChanged: widget.onValueChanged,
+                  onInteractionFinished: widget.onInteractionFinished,
+                  min: widget.minValue,
+                  max: widget.maxValue,
+                  bIsEnabled: widget.bIsEnabled,
+                ),
+                DrivenColorGreenSlider(
+                  key: Key('Green'),
+                  values: widget.wheelValues,
+                  onChanged: widget.onValueChanged,
+                  onInteractionFinished: widget.onInteractionFinished,
+                  min: widget.minValue,
+                  max: widget.maxValue,
+                  bIsEnabled: widget.bIsEnabled,
+                ),
+                DrivenColorBlueSlider(
+                  key: Key('Blue'),
+                  values: widget.wheelValues,
+                  onChanged: widget.onValueChanged,
+                  onInteractionFinished: widget.onInteractionFinished,
+                  min: widget.minValue,
+                  max: widget.maxValue,
+                  bIsEnabled: widget.bIsEnabled,
+                ),
+              ]);
+            }
+
+            break;
+
+          case ColorWheelColorMode.hsv:
+            sliders.addAll([
+              DrivenColorHueSlider(
+                key: Key('Hue'),
+                values: widget.wheelValues,
+                onChanged: widget.onValueChanged,
+                onInteractionFinished: widget.onInteractionFinished,
+                min: widget.minValue,
+                max: widget.maxValue,
+                bIsEnabled: widget.bIsEnabled,
+              ),
+              DrivenColorSaturationSlider(
+                key: Key('Saturation'),
+                values: widget.wheelValues,
+                onChanged: widget.onValueChanged,
+                onInteractionFinished: widget.onInteractionFinished,
+                min: widget.minValue,
+                max: widget.maxValue,
+                bIsEnabled: widget.bIsEnabled,
+              ),
+              DrivenColorValueSlider(
+                key: Key('Value'),
+                values: widget.wheelValues,
+                onChanged: widget.onValueChanged,
+                onInteractionFinished: widget.onInteractionFinished,
+                min: widget.minValue,
+                max: widget.maxValue,
+                bIsEnabled: widget.bIsEnabled,
+              ),
+            ]);
+            break;
+        }
+
+        // Add alpha/luminance slider
+        final Widget? alphaSlider = _makeAlphaSlider();
+        if (alphaSlider != null) {
+          sliders.add(alphaSlider);
+        }
+      }
+
+      // Create color wheel, used in full-height or wheel-only mode
+      if (editMode == ColorWheelEditMode.wheel || isFullHeight) {
+        colorWheel = _ColorWheelWheel(
+          wheelSize: wheelDiameter,
+          uiValues: widget.wheelValues,
+          saturationExponent: widget.saturationExponent,
+          onPanUpdate: _onWheelDragUpdate,
+          onPanEnd: (_) => widget.onInteractionFinished?.call(),
+          onPanCancel: () => widget.onInteractionFinished?.call(),
+        );
+      }
+
+      if (isFullHeight) {
+        mainBody = colorWheel!;
+      } else {
+        switch (editMode) {
+          case ColorWheelEditMode.sliders:
+            mainBody = TransientPreferenceBuilder(
+              preference: bIsInResetModePref,
+              builder: (context, final bool bIsInResetMode) => Padding(
+                // Extra padding at top to fit reset button
+                padding: EdgeInsets.only(top: bIsInResetMode ? 48 : 8),
+                child: Column(
+                  mainAxisAlignment: MainAxisAlignment.start,
+                  children: sliders,
+                ),
+              ),
+            );
+            break;
+
+          case ColorWheelEditMode.wheel:
+            mainBody = colorWheel!;
+
+            // Add value and luminance sliders under the wheel, since we're hiding sliders and those components can't
+            // be controlled by a wheel alone
+            otherWidgets.addAll([
+              DrivenDeltaSlider(
+                label: AppLocalizations.of(context)!.colorComponentValue,
+                defaultValue: (widget.minValue + widget.maxValue) / 2,
+                min: widget.minValue,
+                max: widget.maxValue,
+                onChanged: _handleValueSliderChanged,
+                onInteractionFinished: widget.onInteractionFinished,
+                values: widget.wheelValues
+                    .map((ColorWheelValueData? valueData) =>
+                        valueData != null ? DeltaSliderValueData(value: valueData.value.value) : null)
+                    .toList(),
+                bShowResetButton: false,
+                bIsEnabled: widget.bIsEnabled,
+              ),
+              if (widget.alphaMode == AlphaDisplayMode.luminance) _makeAlphaSlider()!,
+            ]);
+
+            break;
+        }
+      }
+
+      return Column(children: [
+        Container(
+          height: 52,
+          color: Theme.of(context).colorScheme.surfaceTint,
+          padding: EdgeInsets.symmetric(horizontal: 16),
+          child:
+              // Header controls
+              Row(
+            mainAxisAlignment: MainAxisAlignment.start,
+            children: [
+              if (widget.extraHeaderWidget != null) widget.extraHeaderWidget!,
+
+              const Spacer(),
+
+              // Color mode toggle
+              SelectorBar(
+                value: colorMode,
+                onSelected: _onColorModeChangedByUser,
+                valueNames: {
+                  ColorWheelColorMode.rgb: localizations.colorWheelColorModeRGB,
+                  ColorWheelColorMode.hsv: localizations.colorWheelColorModeHSV,
+                },
+              ),
+
+              if (!isFullHeight)
+                // Edit mode toggle
+                Padding(
+                  padding: EdgeInsets.only(left: 16),
+                  child: SizedBox(
+                    width: 36,
+                    height: 36,
+                    child: Tooltip(
+                      message: editMode == ColorWheelEditMode.wheel
+                          ? localizations.colorWheelShowSliders
+                          : localizations.colorWheelShowWheel,
+                      child: IconButton(
+                        onPressed: _toggleEditMode,
+                        color: Colors.transparent,
+                        icon: AssetIcon(
+                          size: 20,
+                          path: editMode == ColorWheelEditMode.wheel
+                              ? 'assets/images/icons/sliders.svg'
+                              : 'assets/images/icons/color_wheel.svg',
+                        ),
+                      ),
+                    ),
+                  ),
+                ),
+            ],
+          ),
+        ),
+        SizedBox(height: 8),
+        if (widget.extraTopWidget != null)
+          Padding(
+            padding: EdgeInsets.only(top: 8, bottom: 8),
+            child: widget.extraTopWidget!,
+          ),
+        Padding(
+          padding: EdgeInsets.symmetric(horizontal: cardMargin),
+          child: Column(
+            children: [
+              // Color wheel
+              Stack(children: [
+                mainBody,
+
+                // Reset button
+                TransientPreferenceBuilder(
+                  preference: bIsInResetModePref,
+                  builder: (context, final bool bIsInResetMode) => bIsInResetMode
+                      ? Align(
+                          alignment: Alignment.topRight,
+                          child: ResetValueButton(onPressed: widget.onReset),
+                        )
+                      : const SizedBox(),
+                ),
+              ]),
+
+              if (isFullHeight)
+                for (Widget slider in sliders) slider,
+
+              if (otherWidgets.length > 0) Column(children: otherWidgets),
+            ],
+          ),
+        ),
+      ]);
+    });
+  }
+
+  /// Make the alpha or luminance slider for the controlled color.
+  Widget? _makeAlphaSlider() {
+    final localizations = AppLocalizations.of(context)!;
+    if (widget.alphaMode != AlphaDisplayMode.hide) {
+      final String alphaLabel = (widget.alphaMode == AlphaDisplayMode.luminance)
+          ? localizations.colorComponentLuminance
+          : localizations.colorComponentAlpha;
+
+      if (widget.createRgbaSlider != null) {
+        return widget.createRgbaSlider!(3, alphaLabel);
+      } else {
+        return DrivenColorAlphaSlider(
+          key: Key('Alpha'),
+          label: alphaLabel,
+          values: widget.wheelValues,
+          onChanged: widget.onValueChanged,
+          onInteractionFinished: widget.onInteractionFinished,
+          min: widget.minValue,
+          max: widget.maxValue,
+        );
+      }
+    }
+
+    return null;
+  }
+
+  /// Called when the color wheel portion of the widget is dragged.
+  void _onWheelDragUpdate(DragUpdateDetails details) {
+    if (!widget.bIsEnabled) {
+      return;
+    }
+
+    final renderBox = context.findRenderObject() as RenderBox?;
+    if (renderBox == null) {
+      return;
+    }
+
+    final Offset scaledDeltaOffset = details.delta * deltaMultiplier / _getWheelDiameter(renderBox.constraints);
+    final List<WheelColor> deltaColors = List.filled(widget.wheelValues.length, WheelColor(scaledDeltaOffset, 0, 0));
+    widget.onValueChanged(deltaColors);
+  }
+
+  /// Called when the user changes the value of the value slider.
+  void _handleValueSliderChanged(List<double> deltaValues) {
+    assert(deltaValues.length == widget.wheelValues.length);
+
+    widget.onValueChanged(
+      deltaValues
+          .map(
+            (delta) => WheelColor(Offset.zero, delta, 0),
+          )
+          .toList(growable: false),
+    );
+  }
+
+  /// Toggle between wheel and slider edit modes.
+  void _toggleEditMode() {
+    final ColorWheelEditMode lastEditMode = _settings.editMode.getValue();
+
+    _settings.editMode.setValue(
+      lastEditMode == ColorWheelEditMode.wheel ? ColorWheelEditMode.sliders : ColorWheelEditMode.wheel,
+    );
+  }
+
+  /// Called when the user changes the color mode.
+  void _onColorModeChangedByUser(ColorWheelColorMode newColorMode) {
+    _settings.colorMode.setValue(newColorMode);
+  }
+
+  /// Return the largest wheel diameter that fits into the widget's constraints.
+  double _getWheelDiameter(BoxConstraints constraints) {
+    final double targetRadius =
+        math.min(math.min(constraints.maxHeight, constraints.maxWidth * _wheelWidthFraction), _maxWheelDiameter);
+
+    if (widget.maxRadius != null) {
+      return math.min(targetRadius, widget.maxRadius!);
+    }
+
+    return targetRadius;
+  }
+}
+
+/// The wheel portion of the color wheel widget, which maps hue and saturation onto a plane.
+class _ColorWheelWheel extends StatelessWidget {
+  const _ColorWheelWheel({
+    required this.wheelSize,
+    required this.uiValues,
+    this.onPanUpdate,
+    this.onPanEnd,
+    this.onPanCancel,
+    required this.saturationExponent,
+    Key? key,
+  }) : super(key: key);
+
+  final double wheelSize;
+  final Function(DragUpdateDetails)? onPanUpdate;
+  final Function(DragEndDetails)? onPanEnd;
+  final Function()? onPanCancel;
+  final List<ColorWheelValueData?> uiValues;
+  final double saturationExponent;
+
+  @override
+  Widget build(BuildContext context) {
+    // Create the widgets for each of the dots on the wheel
+    final List<_ColorWheelValueDot> dotWidgets = [];
+
+    for (int i = 0; i < uiValues.length; ++i) {
+      if (uiValues[i] != null) {
+        dotWidgets.add(_ColorWheelValueDot(
+          valueData: uiValues[i]!,
+          saturationExponent: saturationExponent,
+          label: (uiValues.length > 1) ? (i + 1).toString() : '',
+        ));
+      }
+    }
+
+    final double wheelAndDotSize = wheelSize + _dotDiameter;
+
+    // Use a clipper so we can ignore touches outside of the circle
+    return Center(
+      child: SizedBox(
+        width: wheelAndDotSize,
+        height: wheelAndDotSize,
+        child: Stack(
+          alignment: Alignment.center,
+          children: [
+            // Color wheel
+            SizedBox(
+              width: wheelSize,
+              height: wheelSize,
+              child: CustomPaint(
+                painter: _ColorWheelPainter(context),
+              ),
+            ),
+
+            // Value dots
+            Stack(children: dotWidgets),
+
+            // Gesture detector clipped to circle shape
+            SizedBox(
+              width: wheelAndDotSize,
+              height: wheelAndDotSize,
+              child: ClipOval(
+                clipBehavior: Clip.hardEdge,
+                clipper: _CenteredSquareClipper(),
+                child: GestureDetector(
+                  onPanUpdate: onPanUpdate,
+                  onPanEnd: onPanEnd,
+                  onPanCancel: onPanCancel,
+                  behavior: HitTestBehavior.translucent,
+                ),
+              ),
+            ),
+          ],
+        ),
+      ),
+    );
+  }
+}
+
+/// Dot shown on the color wheel to indicate the value of properties being edited.
+class _ColorWheelValueDot extends StatelessWidget {
+  const _ColorWheelValueDot({
+    required this.valueData,
+    required this.saturationExponent,
+    this.label = '',
+    Key? key,
+  }) : super(key: key);
+
+  final String label;
+  final ColorWheelValueData valueData;
+  final double saturationExponent;
+
+  @override
+  Widget build(BuildContext context) {
+    /// Apply the inverse of the exponent to convert our linear position to the correct position on the circle
+    final Offset position = exponentiateCirclePosition(
+      valueData.value.position,
+      saturationExponent,
+      true,
+    );
+
+    return Align(
+      alignment: Alignment(position.dx, position.dy),
+      child: Container(
+        width: _dotDiameter,
+        height: _dotDiameter,
+        child: Stack(
+          children: [
+            SizedBox(
+              width: _dotDiameter,
+              height: _dotDiameter,
+              child: CustomPaint(
+                painter: _ColorWheelDotPainter(context),
+              ),
+            ),
+          ],
+        ),
+      ),
+    );
+  }
+}
+
+/// Painter for the wheel image of the color wheel.
+class _ColorWheelPainter extends CustomPainter {
+  const _ColorWheelPainter(this.context) : super();
+
+  final BuildContext context;
+
+  /// Paint the wheel widget.
+  @override
+  void paint(Canvas canvas, Size size) {
+    final Rect canvasRect = Offset.zero & size;
+    final double radius = math.min(size.width, size.height) / 2;
+    final Offset center = canvasRect.center;
+
+    const rainbowSweepGradient = SweepGradient(colors: [
+      Color.fromARGB(255, 255, 0, 0),
+      Color.fromARGB(255, 255, 255, 0),
+      Color.fromARGB(255, 0, 255, 0),
+      Color.fromARGB(255, 0, 255, 255),
+      Color.fromARGB(255, 0, 0, 255),
+      Color.fromARGB(255, 255, 0, 255),
+      Color.fromARGB(255, 255, 0, 0),
+    ]);
+
+    const whiteRadialGradient = RadialGradient(colors: [
+      Color.fromARGB(255, 255, 255, 255),
+      Color.fromARGB(0, 255, 255, 255),
+    ]);
+
+    final Paint rainbowSweepPaint = Paint()
+      ..shader = rainbowSweepGradient.createShader(canvasRect)
+      ..style = PaintingStyle.fill;
+
+    final Paint whiteRadialPaint = Paint()
+      ..shader = whiteRadialGradient.createShader(canvasRect)
+      ..style = PaintingStyle.fill;
+
+    canvas.drawCircle(center, radius, rainbowSweepPaint);
+    canvas.drawCircle(center, radius, whiteRadialPaint);
+  }
+
+  /// Only repaint if the widget has been replaced.
+  @override
+  bool shouldRepaint(covariant CustomPainter oldDelegate) {
+    return oldDelegate != this;
+  }
+}
+
+/// Painter for the value dots indicating a value being edited in the color wheel.
+class _ColorWheelDotPainter extends CustomPainter {
+  const _ColorWheelDotPainter(this.context) : super();
+
+  final BuildContext context;
+
+  @override
+  void paint(Canvas canvas, Size size) {
+    const double ringWidth = 1;
+
+    final Rect canvasRect = Offset.zero & size;
+
+    // Draw the ring
+    final Paint ringPaint = Paint()
+      ..style = PaintingStyle.stroke
+      ..strokeWidth = ringWidth
+      ..color = UnrealColors.black;
+
+    final Rect ringRect = canvasRect.inflate(-ringWidth / 2);
+    canvas.drawArc(ringRect, 0, math.pi * 2, false, ringPaint);
+  }
+
+  @override
+  bool shouldRepaint(covariant CustomPainter oldDelegate) {
+    return oldDelegate != this;
+  }
+}
+
+/// A clipper that returns a square in the center of the canvas area.
+class _CenteredSquareClipper extends CustomClipper<Rect> {
+  @override
+  getClip(Size size) {
+    final double squareSize = math.min(size.width, size.height);
+    return Rect.fromCenter(center: size.center(Offset.zero), width: squareSize, height: squareSize);
+  }
+
+  @override
+  bool shouldReclip(covariant CustomClipper<Rect> oldClipper) {
+    return this != oldClipper;
+  }
+}
