Branch snapshot for CL 27138787

diff --git a/Engine/Extras/VirtualProduction/EpicStageApp/lib/widgets/elements/color_slider.dart b/Engine/Extras/VirtualProduction/EpicStageApp/lib/widgets/elements/color_slider.dart
new file mode 100644
--- /dev/null
+++ b/Engine/Extras/VirtualProduction/EpicStageApp/lib/widgets/elements/color_slider.dart
@@ -0,0 +1,665 @@
+// Copyright Epic Games, Inc. All Rights Reserved.
+
+import 'package:flutter/material.dart';
+import 'package:flutter_gen/gen_l10n/app_localizations.dart';
+
+import '../../models/color.dart';
+import '../../models/property_modify_operations.dart';
+import 'delta_slider.dart';
+import 'delta_widget_base.dart';
+import 'unreal_widget_base.dart';
+
+typedef ColorSliderValueData = DeltaWidgetValueData<WheelColor>;
+
+/// A delta-based slider that controls the red component of one or more colors in Unreal Editor.
+class UnrealColorGradingRedSlider extends _UnrealColorComponentSlider {
+  const UnrealColorGradingRedSlider({
+    required super.unrealProperties,
+    required super.enableProperties,
+    required super.colorMin,
+    required super.colorMax,
+    super.key,
+    super.overrideName,
+  });
+
+  @override
+  State<StatefulWidget> createState() => _UnrealColorGradingRedSliderState();
+}
+
+class _UnrealColorGradingRedSliderState extends State<UnrealColorGradingRedSlider>
+    with
+        UnrealWidgetStateMixin<UnrealColorGradingRedSlider, WheelColor>,
+        _UnrealColorComponentSliderState<UnrealColorGradingRedSlider> {
+  @override
+  Widget build(BuildContext context) {
+    return DrivenColorRedSlider(
+      values: makeValueDataList(),
+      onChanged: handleOnChangedByUser,
+      onInteractionFinished: endTransaction,
+      min: widget.colorMin,
+      max: widget.colorMax,
+      label: widget.overrideName,
+    );
+  }
+}
+
+/// A delta-based slider that controls the green component of one or more colors in Unreal Editor.
+class UnrealColorGradingGreenSlider extends _UnrealColorComponentSlider {
+  const UnrealColorGradingGreenSlider({
+    required super.unrealProperties,
+    required super.enableProperties,
+    required super.colorMin,
+    required super.colorMax,
+    super.key,
+    super.overrideName,
+  });
+
+  @override
+  State<StatefulWidget> createState() => _UnrealColorGradingGreenSliderState();
+}
+
+class _UnrealColorGradingGreenSliderState extends State<UnrealColorGradingGreenSlider>
+    with
+        UnrealWidgetStateMixin<UnrealColorGradingGreenSlider, WheelColor>,
+        _UnrealColorComponentSliderState<UnrealColorGradingGreenSlider> {
+  @override
+  Widget build(BuildContext context) {
+    return DrivenColorGreenSlider(
+      values: makeValueDataList(),
+      onChanged: handleOnChangedByUser,
+      onInteractionFinished: endTransaction,
+      min: widget.colorMin,
+      max: widget.colorMax,
+      label: widget.overrideName,
+    );
+  }
+}
+
+/// A delta-based slider that controls the blue component of one or more colors in Unreal Editor.
+class UnrealColorGradingBlueSlider extends _UnrealColorComponentSlider {
+  const UnrealColorGradingBlueSlider({
+    required super.unrealProperties,
+    required super.enableProperties,
+    required super.colorMin,
+    required super.colorMax,
+    super.key,
+    super.overrideName,
+  });
+
+  @override
+  State<StatefulWidget> createState() => _UnrealColorGradingBlueSliderState();
+}
+
+class _UnrealColorGradingBlueSliderState extends State<UnrealColorGradingBlueSlider>
+    with
+        UnrealWidgetStateMixin<UnrealColorGradingBlueSlider, WheelColor>,
+        _UnrealColorComponentSliderState<UnrealColorGradingBlueSlider> {
+  @override
+  Widget build(BuildContext context) {
+    return DrivenColorBlueSlider(
+      values: makeValueDataList(),
+      onChanged: handleOnChangedByUser,
+      onInteractionFinished: endTransaction,
+      min: widget.colorMin,
+      max: widget.colorMax,
+      label: widget.overrideName,
+    );
+  }
+}
+
+/// A delta-based slider that controls the hue component of one or more colors in Unreal Editor.
+class UnrealColorGradingHueSlider extends _UnrealColorComponentSlider {
+  const UnrealColorGradingHueSlider({
+    required super.unrealProperties,
+    required super.enableProperties,
+    required super.colorMin,
+    required super.colorMax,
+    super.key,
+    super.overrideName,
+  });
+
+  @override
+  State<StatefulWidget> createState() => _UnrealColorGradingHueSliderState();
+}
+
+class _UnrealColorGradingHueSliderState extends State<UnrealColorGradingHueSlider>
+    with
+        UnrealWidgetStateMixin<UnrealColorGradingHueSlider, WheelColor>,
+        _UnrealColorComponentSliderState<UnrealColorGradingHueSlider> {
+  @override
+  Widget build(BuildContext context) {
+    return DrivenColorHueSlider(
+      values: makeValueDataList(),
+      onChanged: handleOnChangedByUser,
+      onInteractionFinished: endTransaction,
+      min: widget.colorMin,
+      max: widget.colorMax,
+      label: widget.overrideName,
+    );
+  }
+}
+
+/// A delta-based slider that controls the saturation component of one or more colors in Unreal Editor.
+class UnrealColorGradingSaturationSlider extends _UnrealColorComponentSlider {
+  const UnrealColorGradingSaturationSlider({
+    required super.unrealProperties,
+    required super.enableProperties,
+    required super.colorMin,
+    required super.colorMax,
+    super.key,
+    super.overrideName,
+  });
+
+  @override
+  State<StatefulWidget> createState() => _UnrealColorGradingSaturationSliderState();
+}
+
+class _UnrealColorGradingSaturationSliderState extends State<UnrealColorGradingSaturationSlider>
+    with
+        UnrealWidgetStateMixin<UnrealColorGradingSaturationSlider, WheelColor>,
+        _UnrealColorComponentSliderState<UnrealColorGradingSaturationSlider> {
+  @override
+  Widget build(BuildContext context) {
+    return DrivenColorSaturationSlider(
+      values: makeValueDataList(),
+      onChanged: handleOnChangedByUser,
+      onInteractionFinished: endTransaction,
+      min: widget.colorMin,
+      max: widget.colorMax,
+      label: widget.overrideName,
+    );
+  }
+}
+
+/// A delta-based slider that controls the value component of one or more colors in Unreal Editor.
+class UnrealColorGradingValueSlider extends _UnrealColorComponentSlider {
+  const UnrealColorGradingValueSlider({
+    required super.unrealProperties,
+    required super.enableProperties,
+    required super.colorMin,
+    required super.colorMax,
+    super.key,
+    super.overrideName,
+  });
+
+  @override
+  State<StatefulWidget> createState() => _UnrealColorGradingValueSliderState();
+}
+
+class _UnrealColorGradingValueSliderState extends State<UnrealColorGradingValueSlider>
+    with
+        UnrealWidgetStateMixin<UnrealColorGradingValueSlider, WheelColor>,
+        _UnrealColorComponentSliderState<UnrealColorGradingValueSlider> {
+  @override
+  Widget build(BuildContext context) {
+    return DrivenColorValueSlider(
+      values: makeValueDataList(),
+      onChanged: handleOnChangedByUser,
+      onInteractionFinished: endTransaction,
+      min: widget.colorMin,
+      max: widget.colorMax,
+      label: widget.overrideName,
+    );
+  }
+}
+
+/// A delta-based slider widget that controls the alpha component of one or more color properties.
+class UnrealColorGradingAlphaSlider extends _UnrealColorComponentSlider {
+  const UnrealColorGradingAlphaSlider({
+    required super.unrealProperties,
+    required super.enableProperties,
+    required super.colorMin,
+    required super.colorMax,
+    super.key,
+    super.overrideName,
+  });
+
+  @override
+  State<StatefulWidget> createState() => _UnrealColorGradingAlphaSliderState();
+}
+
+class _UnrealColorGradingAlphaSliderState extends State<UnrealColorGradingAlphaSlider>
+    with
+        UnrealWidgetStateMixin<UnrealColorGradingAlphaSlider, WheelColor>,
+        _UnrealColorComponentSliderState<UnrealColorGradingAlphaSlider> {
+  @override
+  Widget build(BuildContext context) {
+    return DrivenColorAlphaSlider(
+      values: makeValueDataList(),
+      onChanged: handleOnChangedByUser,
+      onInteractionFinished: endTransaction,
+      min: widget.colorMin,
+      max: widget.colorMax,
+      label: widget.overrideName,
+    );
+  }
+}
+
+/// Mixin for state of widgets that extend [_UnrealColorComponentSlider].
+mixin _UnrealColorComponentSliderState<WidgetType extends _UnrealColorComponentSlider>
+    on UnrealWidgetStateMixin<WidgetType, WheelColor> {
+  @override
+  PropertyModifyOperation get modifyOperation => WheelColorGradingAddOperation(
+        saturationExponent: 1.0,
+        minValue: widget.colorMin,
+        maxValue: widget.colorMax,
+      );
+}
+
+/// A delta-based slider widget that controls the red component of one or more color properties.
+class DrivenColorRedSlider extends _DrivenColorComponentSlider {
+  const DrivenColorRedSlider({
+    super.key,
+    required super.values,
+    required super.onChanged,
+    super.bIsEnabled,
+    super.min = 0.0,
+    super.max = 1.0,
+    super.label,
+    super.onInteractionFinished,
+  });
+
+  @override
+  State<StatefulWidget> createState() => _DrivenColorRedSliderState();
+}
+
+class _DrivenColorRedSliderState extends _DrivenRGBComponentSliderState {
+  @override
+  String get defaultLabel => AppLocalizations.of(context)!.colorComponentRed;
+
+  @override
+  Color? get fillColor => Color.fromARGB(255, 255, 0, 0);
+
+  @override
+  double getLinearValue(FloatColor color) => color.red;
+
+  @override
+  FloatColor makeColorWithComponentValue(FloatColor color, double newComponentValue) =>
+      FloatColor(newComponentValue, color.green, color.blue, color.alpha);
+}
+
+/// A delta-based slider widget that controls the green component of one or more color properties.
+class DrivenColorGreenSlider extends _DrivenColorComponentSlider {
+  const DrivenColorGreenSlider({
+    super.key,
+    required super.values,
+    required super.onChanged,
+    super.bIsEnabled,
+    super.min = 0.0,
+    super.max = 1.0,
+    super.label,
+    super.onInteractionFinished,
+  });
+
+  @override
+  State<StatefulWidget> createState() => _DrivenColorGreenSliderState();
+}
+
+class _DrivenColorGreenSliderState extends _DrivenRGBComponentSliderState {
+  @override
+  String get defaultLabel => AppLocalizations.of(context)!.colorComponentGreen;
+
+  @override
+  Color? get fillColor => Color.fromARGB(255, 0, 255, 0);
+
+  @override
+  double getLinearValue(FloatColor color) => color.green;
+
+  @override
+  FloatColor makeColorWithComponentValue(FloatColor color, double newComponentValue) =>
+      FloatColor(color.red, newComponentValue, color.blue, color.alpha);
+}
+
+/// A delta-based slider widget that controls the blue component of one or more color properties.
+class DrivenColorBlueSlider extends _DrivenColorComponentSlider {
+  const DrivenColorBlueSlider({
+    super.key,
+    required super.values,
+    required super.onChanged,
+    super.bIsEnabled,
+    super.min = 0.0,
+    super.max = 1.0,
+    super.label,
+    super.onInteractionFinished,
+  });
+
+  @override
+  State<StatefulWidget> createState() => _DrivenColorBlueSliderState();
+}
+
+class _DrivenColorBlueSliderState extends _DrivenRGBComponentSliderState {
+  @override
+  String get defaultLabel => AppLocalizations.of(context)!.colorComponentBlue;
+
+  @override
+  Color? get fillColor => Color.fromARGB(255, 0, 0, 255);
+
+  @override
+  double getLinearValue(FloatColor color) => color.blue;
+
+  @override
+  FloatColor makeColorWithComponentValue(FloatColor color, double newComponentValue) =>
+      FloatColor(color.red, color.green, newComponentValue, color.alpha);
+}
+
+/// A delta-based slider widget that controls the value component of one or more color properties.
+class DrivenColorHueSlider extends _DrivenColorComponentSlider {
+  const DrivenColorHueSlider({
+    super.key,
+    required super.values,
+    required super.onChanged,
+    super.bIsEnabled,
+    super.min = 0.0,
+    super.max = 1.0,
+    super.label,
+    super.onInteractionFinished,
+  });
+
+  @override
+  State<StatefulWidget> createState() => _DrivenColorHueSliderState();
+}
+
+class _DrivenColorHueSliderState extends State<DrivenColorHueSlider>
+    with DeltaWidgetStateMixin<DrivenColorHueSlider>, _DrivenColorComponentSliderState<DrivenColorHueSlider> {
+  @override
+  String get defaultLabel => AppLocalizations.of(context)!.colorComponentHue;
+
+  @override
+  double get sliderMin => 0.0;
+
+  @override
+  double get sliderMax => 360.0;
+
+  @override
+  NegatableHSVColor applyLinearDelta(NegatableHSVColor currentColor, double delta) {
+    // Apply the change in hue to the wheel color's hue
+    final double deltaHue = delta;
+    final double newHue = (currentColor.hue + deltaHue + sliderMax) % sliderMax;
+
+    // Convert back to a wheel color, then subtract the existing color to produce a delta
+    return currentColor.withHue(newHue);
+  }
+
+  @override
+  double convertToLinearValue(WheelColor color) => color.toHSVColor().hue;
+}
+
+/// A delta-based slider widget that controls the saturation component of one or more color properties.
+class DrivenColorSaturationSlider extends _DrivenColorComponentSlider {
+  const DrivenColorSaturationSlider({
+    super.key,
+    required super.values,
+    required super.onChanged,
+    super.bIsEnabled,
+    super.min = 0.0,
+    super.max = 1.0,
+    super.label,
+    super.onInteractionFinished,
+  });
+
+  @override
+  State<StatefulWidget> createState() => _DrivenColorSaturationSliderState();
+}
+
+class _DrivenColorSaturationSliderState extends _DrivenHSVComponentSliderState {
+  @override
+  double get sliderMin => 0.0;
+
+  @override
+  double get sliderMax => 1.0;
+
+  @override
+  String get defaultLabel => AppLocalizations.of(context)!.colorComponentSaturation;
+
+  @override
+  double getLinearValue(NegatableHSVColor color) => color.saturation;
+
+  @override
+  NegatableHSVColor makeColorWithComponentValue(NegatableHSVColor color, double newComponentValue) =>
+      NegatableHSVColor.fromAHSV(color.alpha, color.hue, newComponentValue, color.value);
+}
+
+/// A delta-based slider widget that controls the value component of one or more color properties.
+class DrivenColorValueSlider extends _DrivenColorComponentSlider {
+  const DrivenColorValueSlider({
+    super.key,
+    required super.values,
+    required super.onChanged,
+    super.bIsEnabled,
+    super.min = 0.0,
+    super.max = 1.0,
+    super.label,
+    super.onInteractionFinished,
+  });
+
+  @override
+  State<StatefulWidget> createState() => _DrivenColorValueSliderState();
+}
+
+class _DrivenColorValueSliderState extends _DrivenHSVComponentSliderState {
+  @override
+  double get sliderMin => widget.min;
+
+  @override
+  double get sliderMax => widget.max;
+
+  @override
+  String get defaultLabel => AppLocalizations.of(context)!.colorComponentValue;
+
+  @override
+  double getLinearValue(NegatableHSVColor color) => color.value;
+
+  @override
+  NegatableHSVColor makeColorWithComponentValue(NegatableHSVColor color, double newComponentValue) =>
+      NegatableHSVColor.fromAHSV(color.alpha, color.hue, color.saturation, newComponentValue);
+}
+
+/// A delta-based slider widget that controls the alpha component of one or more color properties.
+class DrivenColorAlphaSlider extends _DrivenColorComponentSlider {
+  const DrivenColorAlphaSlider({
+    super.key,
+    required super.values,
+    required super.onChanged,
+    super.bIsEnabled,
+    super.min = 0.0,
+    super.max = 1.0,
+    super.label,
+    super.onInteractionFinished,
+  });
+
+  @override
+  State<StatefulWidget> createState() => _DrivenColorAlphaSliderState();
+}
+
+class _DrivenColorAlphaSliderState extends _DrivenHSVComponentSliderState {
+  @override
+  double get sliderMin => widget.min;
+
+  @override
+  double get sliderMax => widget.max;
+
+  @override
+  String get defaultLabel => AppLocalizations.of(context)!.colorComponentAlpha;
+
+  @override
+  double getLinearValue(NegatableHSVColor color) => color.alpha;
+
+  @override
+  NegatableHSVColor makeColorWithComponentValue(NegatableHSVColor color, double newComponentValue) =>
+      NegatableHSVColor.fromAHSV(newComponentValue, color.hue, color.saturation, color.value);
+}
+
+/// Base state class for sliders that operate on RGBA component values.
+abstract class _DrivenRGBComponentSliderState<WidgetType extends _DrivenColorComponentSlider> extends State<WidgetType>
+    with DeltaWidgetStateMixin<WidgetType>, _DrivenColorComponentSliderState<WidgetType> {
+  @override
+  double get sliderMin => widget.min;
+
+  @override
+  double get sliderMax => widget.max;
+
+  /// Get the linear value this slider modifies from the given color.
+  double getLinearValue(FloatColor color);
+
+  /// Create a new color by replacing the value that this slider modifies.
+  FloatColor makeColorWithComponentValue(FloatColor color, double newComponentValue);
+
+  @override
+  NegatableHSVColor applyLinearDelta(NegatableHSVColor currentColor, double delta) {
+    // Apply the change to the relevant component
+    final currentRGBA = FloatColor.fromHSVColor(currentColor);
+    final double newValue = (getLinearValue(currentRGBA) + delta).clamp(sliderMin, sliderMax);
+
+    return NegatableHSVColor.fromColor(makeColorWithComponentValue(currentRGBA, newValue));
+  }
+
+  @override
+  double convertToLinearValue(WheelColor color) => getLinearValue(color.toFloatColor());
+}
+
+/// Base state class for sliders that operate on clamped HSV component values.
+abstract class _DrivenHSVComponentSliderState<WidgetType extends _DrivenColorComponentSlider> extends State<WidgetType>
+    with DeltaWidgetStateMixin<WidgetType>, _DrivenColorComponentSliderState<WidgetType> {
+  /// Get the linear value this slider modifies from the given color.
+  double getLinearValue(NegatableHSVColor color);
+
+  /// Create a new color by replacing the value that this slider modifies.
+  NegatableHSVColor makeColorWithComponentValue(NegatableHSVColor color, double newComponentValue);
+
+  @override
+  NegatableHSVColor applyLinearDelta(NegatableHSVColor currentColor, double delta) {
+    // Apply the change to the relevant component
+    final double newValue = (getLinearValue(currentColor) + delta).clamp(sliderMin, sliderMax);
+
+    return makeColorWithComponentValue(currentColor, newValue);
+  }
+
+  @override
+  double convertToLinearValue(WheelColor color) => getLinearValue(color.toHSVColor());
+}
+
+/// Base class for sliders that control a component of a color that needs to be calculated based on the full color data,
+/// and which control these properties in Unreal Editor.
+abstract class _UnrealColorComponentSlider extends UnrealWidget {
+  const _UnrealColorComponentSlider({
+    required super.unrealProperties,
+    required super.enableProperties,
+    super.key,
+    super.overrideName,
+    required this.colorMin,
+    required this.colorMax,
+  });
+
+  /// The minimum value of the color's components.
+  final double colorMin;
+
+  /// The maximum value of the color's components.
+  final double colorMax;
+}
+
+/// Base class for sliders that control a component of a color that needs to be calculated based on the full color data.
+abstract class _DrivenColorComponentSlider extends StatefulWidget {
+  const _DrivenColorComponentSlider({
+    Key? key,
+    required this.values,
+    required this.onChanged,
+    this.bIsEnabled = true,
+    this.min = 0.0,
+    this.max = 1.0,
+    this.label,
+    this.onInteractionFinished,
+  }) : super(key: key);
+
+  /// Colors to display on the slider.
+  final List<ColorSliderValueData?> values;
+
+  /// Name of the property this widget controls.
+  final String? label;
+
+  /// If false, grey out the widget and disable its controls.
+  final bool bIsEnabled;
+
+  /// The minimum value of the color's components.
+  final double min;
+
+  /// The maximum value of the color's components.
+  final double max;
+
+  /// Function called when the user changes the value of the slider.
+  /// Passes the amount by which each value changed.
+  final void Function(List<WheelColor>) onChanged;
+
+  /// Function called when the user is done interacting with the widget.
+  final void Function()? onInteractionFinished;
+}
+
+/// Mixin for state of widgets that extend [_DrivenColorComponentSlider].
+mixin _DrivenColorComponentSliderState<WidgetType extends _DrivenColorComponentSlider>
+    on DeltaWidgetStateMixin<WidgetType> {
+  /// The label to show on the slider if no override is provided.
+  String get defaultLabel;
+
+  /// The minimum value of the color's components.
+  double get sliderMin;
+
+  /// The maximum value of the color's components.
+  double get sliderMax;
+
+  /// The fill color of the slider.
+  Color? get fillColor => null;
+
+  @override
+  Widget build(BuildContext context) {
+    return DrivenDeltaSlider(
+      values: _makeSliderValueList(),
+      bShowResetButton: false,
+      onChanged: _handleOnSliderChangedByUser,
+      onInteractionFinished: widget.onInteractionFinished,
+      min: sliderMin,
+      max: sliderMax,
+      label: widget.label ?? defaultLabel,
+      fillColor: fillColor,
+      bIsEnabled: widget.bIsEnabled,
+    );
+  }
+
+  /// Convert a color value to the linear value the slider controls.
+  double convertToLinearValue(WheelColor color);
+
+  /// Apply the linear delta value from the slider to the calculated component of the color.
+  NegatableHSVColor applyLinearDelta(NegatableHSVColor currentColor, double delta);
+
+  /// Generate a list of linear values corresponding to each controlled property.
+  List<DeltaWidgetValueData<double>?> _makeSliderValueList() {
+    return widget.values.map((final DeltaWidgetValueData<WheelColor>? color) {
+      if (color == null) {
+        return null;
+      }
+
+      return DeltaWidgetValueData<double>(value: convertToLinearValue(color.value));
+    }).toList(growable: false);
+  }
+
+  /// Called when the user changes the controlled linear values.
+  void _handleOnSliderChangedByUser(List<double> deltaValues) {
+    assert(deltaValues.length == widget.values.length);
+
+    final List<DeltaWidgetValueData<WheelColor>?> currentColors = widget.values;
+    final List<WheelColor> colorDeltas = [];
+
+    for (int propertyIndex = 0; propertyIndex < deltaValues.length; ++propertyIndex) {
+      final WheelColor? currentColor = currentColors[propertyIndex]?.value;
+      if (currentColor == null) {
+        colorDeltas.add(WheelColor.zero);
+        continue;
+      }
+
+      // Apply the change in hue, then subtract the old color from the new one to determine the delta
+      final NegatableHSVColor newHsvColor = applyLinearDelta(currentColor.toHSVColor(), deltaValues[propertyIndex]);
+      final WheelColor newColor = WheelColor.fromHSVColor(newHsvColor);
+      colorDeltas.add(newColor - currentColor);
+    }
+
+    widget.onChanged(colorDeltas);
+  }
+}
diff --git a/Engine/Extras/VirtualProduction/EpicStageApp/lib/widgets/elements/control_lock.dart b/Engine/Extras/VirtualProduction/EpicStageApp/lib/widgets/elements/control_lock.dart
new file mode 100644
--- /dev/null
+++ b/Engine/Extras/VirtualProduction/EpicStageApp/lib/widgets/elements/control_lock.dart
@@ -0,0 +1,138 @@
+// Copyright Epic Games, Inc. All Rights Reserved.
+
+import 'package:flutter/material.dart';
+import 'package:flutter/services.dart';
+import 'package:flutter_gen/gen_l10n/app_localizations.dart';
+
+import '../../utilities/unreal_colors.dart';
+import 'epic_icon_button.dart';
+
+/// Widget housing for locking or unlocking the screen, either ignoring touch input on allowing them.
+class ControlLock extends StatefulWidget {
+  const ControlLock({Key? key}) : super(key: key);
+
+  @override
+  State<ControlLock> createState() => _ControlLockState();
+}
+
+class _ControlLockState extends State<ControlLock> {
+  /// Global key for the lock button, allows easy reference to the button.
+  final GlobalKey _buttonKey = GlobalKey();
+
+  /// whether the screen is locked or not.
+  bool _bIsLocked = false;
+
+  /// The context of the overlay shown when locked.
+  BuildContext? _lockOverlayContext;
+
+  @override
+  Widget build(BuildContext context) {
+    return _ControlLockButton(
+      key: _buttonKey,
+      bIsLocked: false,
+      onLongPressed: _onLocked,
+      bIsVisible: !_bIsLocked,
+    );
+  }
+
+  /// callback to initiate screen lock and negate all touch input.
+  void _onLocked() {
+    if (_bIsLocked) {
+      return;
+    }
+
+    final RenderBox? buttonBox = _buttonKey.currentContext?.findRenderObject() as RenderBox;
+
+    if (buttonBox == null) {
+      return;
+    }
+
+    showDialog(
+      context: context,
+      builder: (context) {
+        final Offset buttonPosition = buttonBox.localToGlobal(Offset.zero);
+        _lockOverlayContext = context;
+
+        // Create the overlay
+        return Material(
+          color: Colors.transparent,
+          surfaceTintColor: Colors.transparent,
+          child: Stack(
+            children: [
+              Positioned(
+                left: buttonPosition.dx,
+                child: _ControlLockButton(
+                  bIsLocked: true,
+                  onLongPressed: _onUnlocked,
+                ),
+              ),
+            ],
+          ),
+        );
+      },
+      barrierDismissible: false,
+      barrierColor: const Color(0xC0000000),
+    );
+
+    HapticFeedback.mediumImpact();
+
+    setState(() {
+      _bIsLocked = true;
+    });
+  }
+
+  /// callback to initiate screen unlock and restore touch input on the app.
+  void _onUnlocked() {
+    if (!_bIsLocked) {
+      return;
+    }
+
+    if (_lockOverlayContext != null) {
+      Navigator.of(_lockOverlayContext!).pop();
+      _lockOverlayContext = null;
+    }
+
+    HapticFeedback.mediumImpact();
+
+    setState(() {
+      _bIsLocked = false;
+    });
+  }
+}
+
+/// lock button
+class _ControlLockButton extends StatelessWidget {
+  const _ControlLockButton({
+    Key? key,
+    required this.onLongPressed,
+    required this.bIsLocked,
+    this.bIsVisible = true,
+  }) : super(key: key);
+
+  /// Function to be executed when the button is long pressed.
+  final Function() onLongPressed;
+
+  /// whether the button is locked or not.
+  final bool bIsLocked;
+
+  /// Whether the lock icon on the button is visible when the button is locked vs opened.
+  final bool bIsVisible;
+
+  @override
+  Widget build(BuildContext context) {
+    return SizedBox(
+      width: 48,
+      height: 48,
+      child: !bIsVisible
+          ? null
+          : EpicIconButton(
+              bIsToggledOn: bIsLocked,
+              activeBackgroundColor: UnrealColors.highlightRed,
+              iconPath: bIsLocked ? 'assets/images/icons/lock_locked.svg' : 'assets/images/icons/lock_unlocked.svg',
+              onLongPressed: onLongPressed,
+              tooltipMessage:
+                  bIsLocked ? AppLocalizations.of(context)!.unlockControls : AppLocalizations.of(context)!.lockControls,
+            ),
+    );
+  }
+}
diff --git a/Engine/Extras/VirtualProduction/EpicStageApp/lib/widgets/elements/cupertino/cupertino_checkbox.dart b/Engine/Extras/VirtualProduction/EpicStageApp/lib/widgets/elements/cupertino/cupertino_checkbox.dart
new file mode 100644
--- /dev/null
+++ b/Engine/Extras/VirtualProduction/EpicStageApp/lib/widgets/elements/cupertino/cupertino_checkbox.dart
@@ -0,0 +1,92 @@
+// Copyright Epic Games, Inc. All Rights Reserved.
+
+import 'package:flutter/cupertino.dart';
+
+/// Radius of the checkbox' circle.
+const double _checkboxRadius = 11.0;
+
+/// Thickness of the ring around an inactive checkbox.
+const double _checkboxRingWidth = 1.5;
+
+/// Size of the check icon relative to the radius of the checkbox.
+const double _checkboxIconSize = 1.2;
+
+/// Offset added to the position of the check icon.
+const Offset _checkboxIconOffset = Offset(0.0, -1);
+
+/// A circular checkbox in the Cupertino style.
+class CupertinoCheckbox extends StatelessWidget {
+  const CupertinoCheckbox({required this.bIsChecked, this.onChanged, Key? key}) : super(key: key);
+
+  /// Whether the box is checked.
+  final bool bIsChecked;
+
+  /// Called when the value of the checkbox should change.
+  final ValueChanged<bool>? onChanged;
+
+  @override
+  Widget build(BuildContext context) {
+    return CustomPaint(
+      painter: _CupertinoCheckboxPainter(
+        bIsChecked: bIsChecked,
+        inactiveColor: CupertinoDynamicColor.resolve(CupertinoColors.inactiveGray, context),
+        activeColor: CupertinoDynamicColor.resolve(CupertinoColors.activeBlue, context),
+        checkColor: CupertinoDynamicColor.resolve(CupertinoColors.white, context),
+      ),
+    );
+  }
+}
+
+class _CupertinoCheckboxPainter extends CustomPainter {
+  const _CupertinoCheckboxPainter({
+    required this.bIsChecked,
+    required this.inactiveColor,
+    required this.activeColor,
+    required this.checkColor,
+  });
+
+  final bool bIsChecked;
+  final Color inactiveColor;
+  final Color activeColor;
+  final Color checkColor;
+
+  @override
+  void paint(Canvas canvas, Size size) {
+    final Offset center = size.center(Offset.zero);
+
+    if (bIsChecked) {
+      // Filled circle and checkmark when checked
+      final Paint circlePaint = Paint()..color = activeColor;
+
+      canvas.drawCircle(center, _checkboxRadius, circlePaint);
+
+      const IconData icon = CupertinoIcons.check_mark;
+      TextPainter textPainter = TextPainter(textDirection: TextDirection.ltr)
+        ..text = TextSpan(
+          text: String.fromCharCode(icon.codePoint),
+          style: TextStyle(
+            color: checkColor,
+            fontSize: _checkboxRadius * _checkboxIconSize,
+            fontFamily: icon.fontFamily,
+            package: icon.fontPackage,
+          ),
+        )
+        ..layout();
+
+      textPainter.paint(canvas, center - textPainter.size.center(_checkboxIconOffset));
+    } else {
+      // Unfilled circle when unchecked
+      final Paint ringPaint = Paint()
+        ..color = inactiveColor
+        ..style = PaintingStyle.stroke
+        ..strokeWidth = _checkboxRingWidth;
+
+      canvas.drawCircle(center, _checkboxRadius - (_checkboxRingWidth / 2), ringPaint);
+    }
+  }
+
+  @override
+  bool shouldRepaint(covariant CustomPainter oldDelegate) {
+    return this != oldDelegate;
+  }
+}
diff --git a/Engine/Extras/VirtualProduction/EpicStageApp/lib/widgets/elements/cupertino/cupertino_icon_button.dart b/Engine/Extras/VirtualProduction/EpicStageApp/lib/widgets/elements/cupertino/cupertino_icon_button.dart
new file mode 100644
--- /dev/null
+++ b/Engine/Extras/VirtualProduction/EpicStageApp/lib/widgets/elements/cupertino/cupertino_icon_button.dart
@@ -0,0 +1,72 @@
+// Copyright Epic Games, Inc. All Rights Reserved.
+
+import 'package:flutter/cupertino.dart';
+import 'package:flutter/material.dart';
+
+/// A button with no fill and a single colored app_icon visible.
+class CupertinoIconButton extends StatelessWidget {
+  const CupertinoIconButton(
+      {this.bEnabled = true,
+      this.iconSize,
+      this.description = '',
+      this.color,
+      this.onPressed,
+      this.icon = CupertinoIcons.question,
+      this.xPadding,
+      this.pressedOpacity = 0.4,
+      Key? key})
+      : super(key: key);
+
+  /// Size of the icon.
+  final double? iconSize;
+
+  /// Color of the app_icon.
+  final Color? color;
+
+  /// Whether the button is enabled. If false, events won't be called and the app_icon will be greyed out.
+  final bool bEnabled;
+
+  /// Callback for when the button is pressed.
+  final void Function()? onPressed;
+
+  /// The app_icon to display.
+  final IconData icon;
+
+  /// Description of the button shown as a tooltip when hovered/long-pressed.
+  final String description;
+
+  /// The opacity that the button will fade to when pressed.
+  final double? pressedOpacity;
+
+  /// Padding along the x-axis on either side.
+  final double? xPadding;
+
+  @override
+  Widget build(BuildContext context) {
+    Color colorToShow;
+    if (bEnabled) {
+      if (color != null) {
+        colorToShow = color!;
+      } else {
+        colorToShow = CupertinoTheme.of(context).primaryColor;
+      }
+    } else {
+      colorToShow = CupertinoDynamicColor.resolve(CupertinoColors.inactiveGray, context);
+    }
+
+    final double size = iconSize ?? const CupertinoIconThemeData().resolve(context).size ?? 24.0;
+
+    return SizedBox(
+        width: size + ((xPadding ?? 6) * 2),
+        height: size,
+        child: Tooltip(
+            message: description,
+            waitDuration: const Duration(seconds: 1),
+            child: CupertinoButton(
+              padding: EdgeInsets.zero,
+              onPressed: bEnabled ? onPressed : null,
+              pressedOpacity: pressedOpacity,
+              child: Icon(icon, size: size, color: colorToShow),
+            )));
+  }
+}
diff --git a/Engine/Extras/VirtualProduction/EpicStageApp/lib/widgets/elements/cupertino/cupertino_pulldown.dart b/Engine/Extras/VirtualProduction/EpicStageApp/lib/widgets/elements/cupertino/cupertino_pulldown.dart
new file mode 100644
--- /dev/null
+++ b/Engine/Extras/VirtualProduction/EpicStageApp/lib/widgets/elements/cupertino/cupertino_pulldown.dart
@@ -0,0 +1,415 @@
+// Copyright Epic Games, Inc. All Rights Reserved.
+
+import 'dart:ui';
+
+import 'package:flutter/cupertino.dart';
+
+/// Width of an iOS-style modal.
+const double _modalWidth = 250;
+
+/// Color of the barrier behind the modal.
+const Color _modalBarrierColor = Color(0x00000000);
+
+/// Corner radius of the pull-down menu's shape so it can cast a matching shadow.
+const double _modalCornerRadius = 14.0;
+
+const Duration _transitionInDuration = Duration(milliseconds: 335);
+const Duration _transitionOutDuration = Duration(milliseconds: 275);
+
+/// Data describing an action to display in an iOS-style pull-down menu.
+class CupertinoPullDownActionData {
+  const CupertinoPullDownActionData({
+    required this.title,
+    this.trailingIcon,
+    this.bIsDestructive = false,
+    this.bCanBeToggled = false,
+    this.bIsToggled = false,
+    this.onPressed,
+    Key? key,
+  }) : assert(bCanBeToggled || !bIsToggled);
+
+  /// The title of the action.
+  final String title;
+
+  /// Optional icon to display after the title.
+  final IconData? trailingIcon;
+
+  /// Whether to style this as a destructive action.
+  final bool bIsDestructive;
+
+  /// If true, leave space next to the title for a checkmark.
+  final bool bCanBeToggled;
+
+  /// If true and bCanBeToggled is set, show a checkmark next to the title.
+  final bool bIsToggled;
+
+  /// Called when the action is pressed.
+  final VoidCallback? onPressed;
+}
+
+/// An iOS-style pull-down button that exposes a menu when pressed.
+/// See https://developer.apple.com/design/human-interface-guidelines/ios/controls/buttons/#pull-down-buttons
+class CupertinoPullDownButton extends StatefulWidget {
+  const CupertinoPullDownButton({required this.child, required this.actionGroups, Key? key}) : super(key: key);
+
+  /// The contents of the button widget.
+  final Widget child;
+
+  /// Widgets to display in the action list. The top-level list is a list of groups which will have thick dividers
+  /// between them, while the inner lists are lists of action widgets to display in each group.
+  final List<List<CupertinoPullDownActionData>> actionGroups;
+
+  @override
+  State<CupertinoPullDownButton> createState() => _CupertinoPullDownButtonState();
+}
+
+class _CupertinoPullDownButtonState extends State<CupertinoPullDownButton> {
+  @override
+  Widget build(BuildContext context) {
+    return CupertinoButton(
+      child: widget.child,
+      onPressed: _showPullDownMenu,
+      padding: const EdgeInsets.symmetric(vertical: 16.0),
+    );
+  }
+
+  /// Open the pull-down menu in a new navigator context.
+  void _showPullDownMenu() {
+    // Determine the position of the button in global coordinates
+    final NavigatorState navigator = Navigator.of(context);
+    final RenderBox buttonBox = context.findRenderObject()! as RenderBox;
+    final Rect buttonRect =
+        buttonBox.localToGlobal(Offset.zero, ancestor: navigator.context.findRenderObject()) & buttonBox.size;
+
+    final route = _PullDownMenuRoute(
+      buttonRect: buttonRect,
+      actionGroups: widget.actionGroups,
+    );
+
+    Navigator.of(context, rootNavigator: true).push<void>(route);
+  }
+}
+
+/// An action shown in the pull-down list of a [CupertinoPullDownButton].
+class _CupertinoPullDownAction extends StatefulWidget {
+  const _CupertinoPullDownAction({
+    required this.actionData,
+    this.bShowTogglePadding = false,
+    Key? key,
+  }) : super(key: key);
+
+  /// The data describing this action.
+  final CupertinoPullDownActionData actionData;
+
+  /// If true, padding will be added on the left side to make room for a toggle checkmark.
+  final bool bShowTogglePadding;
+
+  @override
+  State<_CupertinoPullDownAction> createState() => _CupertinoPullDownActionState();
+}
+
+class _CupertinoPullDownActionState extends State<_CupertinoPullDownAction> {
+  // Color constants based on screenshots from iOS
+  static const Color _backgroundColor = Color(0xC0212122);
+  static const Color _backgroundColorPressed = Color(0xC03F3F40);
+  static const double _buttonHeight = 44.0;
+  static const double _buttonVerticalPadding = 8.0;
+
+  static const TextStyle _titleTextStyle = TextStyle(
+    fontFamily: '.SF UI Text',
+    inherit: false,
+    fontSize: 16.0,
+    textBaseline: TextBaseline.alphabetic,
+  );
+
+  /// Whether the tap/mouse is current pressing on this.
+  bool _bIsPressed = false;
+
+  /// Whether this was tapped, meaning it's been selected and the pull-down will close.
+  bool _bWasTapped = false;
+
+  bool get _bShouldHighlight => _bIsPressed || _bWasTapped;
+
+  @override
+  Widget build(BuildContext context) {
+    final Color textColor = widget.actionData.bIsDestructive
+        ? CupertinoDynamicColor.resolve(CupertinoColors.destructiveRed, context)
+        : CupertinoTheme.of(context).primaryContrastingColor;
+
+    return GestureDetector(
+      behavior: HitTestBehavior.opaque,
+      onTapDown: _onTapDown,
+      onTapUp: _onTapUp,
+      onTapCancel: _onTapCancel,
+      onTap: _onTap,
+      child: ConstrainedBox(
+        constraints: const BoxConstraints(minHeight: _buttonHeight),
+        child: Semantics(
+          button: true,
+          child: Container(
+            decoration: BoxDecoration(
+              color: CupertinoDynamicColor.resolve(
+                  _bShouldHighlight ? _backgroundColorPressed : _backgroundColor, context),
+            ),
+            padding: EdgeInsets.fromLTRB(
+              widget.bShowTogglePadding ? 10.0 : 20.0,
+              _buttonVerticalPadding,
+              widget.actionData.trailingIcon != null ? 16.0 : 20.0,
+              _buttonVerticalPadding,
+            ),
+            child: IconTheme(
+              data: IconThemeData(color: textColor),
+              child: DefaultTextStyle(
+                style: _titleTextStyle.copyWith(color: textColor),
+                child: Row(
+                  children: [
+                    if (widget.bShowTogglePadding)
+                      SizedBox(
+                        width: 22.0,
+                        child: widget.actionData.bIsToggled
+                            ? const Align(
+                                alignment: Alignment.centerLeft,
+                                child: Icon(
+                                  CupertinoIcons.check_mark,
+                                  size: 16.0,
+                                ),
+                              )
+                            : null,
+                      ),
+                    Expanded(child: Text(widget.actionData.title)),
+                    if (widget.actionData.trailingIcon != null)
+                      Icon(
+                        widget.actionData.trailingIcon!,
+                        size: 20.0,
+                      ),
+                  ],
+                ),
+              ),
+            ),
+          ),
+        ),
+      ),
+    );
+  }
+
+  void _onTapDown(TapDownDetails details) {
+    setState(() {
+      _bIsPressed = true;
+    });
+  }
+
+  void _onTapUp(TapUpDetails details) {
+    setState(() {
+      _bIsPressed = false;
+    });
+  }
+
+  void _onTapCancel() {
+    setState(() {
+      _bIsPressed = false;
+    });
+  }
+
+  void _onTap() {
+    setState(() {
+      _bWasTapped = true;
+    });
+
+    widget.actionData.onPressed?.call();
+
+    // Close the pull-down menu
+    Navigator.of(context).pop();
+  }
+}
+
+/// A route that displays the pull-down menu popup, positions it, and animates it into place.
+class _PullDownMenuRoute<T> extends PopupRoute<T> {
+  _PullDownMenuRoute({required this.buttonRect, required this.actionGroups});
+
+  final Rect buttonRect;
+  final List<List<CupertinoPullDownActionData>> actionGroups;
+
+  @override
+  Color? get barrierColor => _modalBarrierColor;
+
+  @override
+  bool get barrierDismissible => true;
+
+  @override
+  String? get barrierLabel => 'Dismiss';
+
+  @override
+  Widget buildPage(BuildContext context, Animation<double> animation, Animation<double> secondaryAnimation) {
+    return Builder(
+      builder: (BuildContext context) {
+        return CustomSingleChildLayout(
+          delegate: _PullDownMenuRouteLayout<T>(
+            buttonRect: buttonRect,
+          ),
+          child: _PullDownMenu(
+            actionGroups: actionGroups,
+            buttonRect: buttonRect,
+            animation: animation,
+          ),
+        );
+      },
+    );
+  }
+
+  @override
+  Duration get transitionDuration => _transitionInDuration;
+
+  @override
+  Duration get reverseTransitionDuration => _transitionOutDuration;
+}
+
+class _PullDownMenuRouteLayout<T> extends SingleChildLayoutDelegate {
+  _PullDownMenuRouteLayout({required this.buttonRect});
+
+  final Rect buttonRect;
+
+  @override
+  BoxConstraints getConstraintsForChild(BoxConstraints constraints) {
+    return const BoxConstraints(minWidth: _modalWidth, maxWidth: _modalWidth);
+  }
+
+  @override
+  Offset getPositionForChild(Size size, Size childSize) {
+    return buttonRect.bottomRight + const Offset(-_modalWidth, 6);
+  }
+
+  @override
+  bool shouldRelayout(covariant SingleChildLayoutDelegate oldDelegate) {
+    return oldDelegate != this;
+  }
+}
+
+/// A pull-down menu displayed when a pull-down button is pressed.
+class _PullDownMenu extends StatelessWidget {
+  const _PullDownMenu({required this.actionGroups, required this.buttonRect, required this.animation, Key? key})
+      : super(key: key);
+
+  final List<List<CupertinoPullDownActionData>> actionGroups;
+  final Rect buttonRect;
+  final Animation<double> animation;
+
+  static const _borderRadius = BorderRadius.all(Radius.circular(_modalCornerRadius));
+
+  static final CurveTween _scaleInTween = CurveTween(curve: const Cubic(0.28, 0.67, 0.4, 1.15));
+  static final CurveTween _scaleOutTween = CurveTween(curve: Curves.easeInQuad);
+  static final CurveTween _fadeInTween = CurveTween(curve: Curves.easeOutCubic);
+  static final CurveTween _fadeOutTween = CurveTween(curve: const Interval(0.2, 1.0, curve: Curves.easeInQuad));
+  static final CurveTween _blurInTween = CurveTween(curve: const Interval(0.0, 0.6, curve: Cubic(0.5, 0.0, 1.0, 0.5)));
+  static final CurveTween _blurOutTween = CurveTween(curve: Curves.easeInQuad);
+
+  @override
+  Widget build(BuildContext context) {
+    return AnimatedBuilder(
+      animation: animation,
+      child: _buildActionColumn(),
+      builder: (BuildContext context, Widget? child) {
+        final bool bIsReversing = animation.status == AnimationStatus.reverse;
+        final double scale = (bIsReversing ? _scaleOutTween : _scaleInTween).evaluate(animation);
+        final double opacity = (bIsReversing ? _fadeOutTween : _fadeInTween).evaluate(animation);
+        final double blurAmount = (bIsReversing ? _blurOutTween : _blurInTween).evaluate(animation) * 10.0;
+
+        return Transform.scale(
+          alignment: Alignment.topRight,
+          origin: Offset(-buttonRect.width / 2, 0),
+          scale: scale,
+          child: Container(
+            decoration: BoxDecoration(
+              borderRadius: _borderRadius,
+              boxShadow: [
+                BoxShadow(color: Color.fromRGBO(0, 0, 0, opacity * 0.2), blurRadius: 20.0, spreadRadius: 25.0)
+              ],
+            ),
+            child: ClipRRect(
+              borderRadius: _borderRadius,
+              // BackdropFilter doesn't work when inside an Opacity widget, so we need to animate its blur separately.
+              // As a result, we re-implement CupertinoPopupSurface here so we can control the blur directly.
+              child: BackdropFilter(
+                filter: ImageFilter.blur(sigmaX: blurAmount, sigmaY: blurAmount),
+                child: Opacity(
+                  opacity: opacity,
+                  child: child,
+                ),
+              ),
+            ),
+          ),
+        );
+      },
+    );
+  }
+
+  /// Build a column of action widgets and separators based on the nested groups passed in.
+  Widget _buildActionColumn() {
+    List<Widget> separatedActions = [];
+
+    const thinSeparator = _ThinPullDownSeparator();
+    const thickSeparator = _ThickPullDownSeparator();
+
+    // First, check if any of the actions can be toggled so we know whether need the extra padding.
+    bool bCanAnyBeToggled = false;
+    for (List<CupertinoPullDownActionData> actionGroup in actionGroups) {
+      for (CupertinoPullDownActionData actionData in actionGroup) {
+        if (actionData.bCanBeToggled) {
+          bCanAnyBeToggled = true;
+          break;
+        }
+      }
+
+      if (bCanAnyBeToggled) {
+        break;
+      }
+    }
+
+    // Build a single list of widgets to display, including separators, based on the nested groups
+    for (int groupIndex = 0; groupIndex < actionGroups.length; ++groupIndex) {
+      List<CupertinoPullDownActionData> actionGroup = actionGroups[groupIndex];
+
+      if (groupIndex != 0) {
+        separatedActions.add(thickSeparator);
+      }
+
+      for (int actionIndex = 0; actionIndex < actionGroup.length; ++actionIndex) {
+        if (actionIndex != 0) {
+          separatedActions.add(thinSeparator);
+        }
+
+        CupertinoPullDownActionData actionData = actionGroup[actionIndex];
+        separatedActions.add(_CupertinoPullDownAction(actionData: actionData, bShowTogglePadding: bCanAnyBeToggled));
+      }
+    }
+
+    return Column(mainAxisSize: MainAxisSize.min, children: separatedActions);
+  }
+}
+
+/// A thin separator between pull-down actions in the same group.
+class _ThinPullDownSeparator extends StatelessWidget {
+  const _ThinPullDownSeparator({Key? key}) : super(key: key);
+
+  @override
+  Widget build(BuildContext context) {
+    return Container(
+      height: 1.0 / MediaQuery.of(context).devicePixelRatio,
+      color: CupertinoDynamicColor.resolve(CupertinoColors.opaqueSeparator.withOpacity(0.3), context),
+    );
+  }
+}
+
+/// A thick separator between pull-down action groups.
+class _ThickPullDownSeparator extends StatelessWidget {
+  const _ThickPullDownSeparator({Key? key}) : super(key: key);
+
+  static const _groupDividerColor = Color(0xC0181818);
+
+  @override
+  Widget build(BuildContext context) {
+    return Container(
+      height: 8.0,
+      color: CupertinoDynamicColor.resolve(_groupDividerColor, context),
+    );
+  }
+}
diff --git a/Engine/Extras/VirtualProduction/EpicStageApp/lib/widgets/elements/delta_slider.dart b/Engine/Extras/VirtualProduction/EpicStageApp/lib/widgets/elements/delta_slider.dart
new file mode 100644
--- /dev/null
+++ b/Engine/Extras/VirtualProduction/EpicStageApp/lib/widgets/elements/delta_slider.dart
@@ -0,0 +1,652 @@
+// Copyright Epic Games, Inc. All Rights Reserved.
+
+import 'dart:math' as math;
+
+import 'package:flutter/material.dart';
+import 'package:flutter_gen/gen_l10n/app_localizations.dart';
+import 'package:provider/provider.dart';
+
+import '../../models/property_modify_operations.dart';
+import '../../models/settings/delta_widget_settings.dart';
+import '../../utilities/math_utilities.dart';
+import '../../utilities/transient_preference.dart';
+import '../../utilities/unreal_colors.dart';
+import 'delta_widget_base.dart';
+import 'modal.dart';
+import 'unreal_widget_base.dart';
+
+const double _dotSize = 17.0;
+const double _trackPadding = DeltaWidgetConstants.widgetOuterXPadding + DeltaWidgetConstants.widgetInnerXPadding;
+const double _trackHeight = 3.0;
+
+typedef DeltaSliderValueData = DeltaWidgetValueData<double>;
+
+/// A delta-based slider widget that controls a property remotely in Unreal Engine.
+class UnrealDeltaSlider extends UnrealWidget {
+  const UnrealDeltaSlider({
+    super.key,
+    super.overrideName,
+    super.minMaxBehaviour,
+    super.enableProperties,
+    required super.unrealProperties,
+    this.buildLabel,
+    this.softMin,
+    this.softMax,
+    this.hardMin,
+    this.hardMax,
+    this.bShowResetButton = true,
+    this.fillColor,
+  });
+
+  /// If provided, call this function to build the widget displayed as the slider's label. Otherwise, the label will be
+  /// automatically generated.
+  final Widget Function(String name)? buildLabel;
+
+  /// If provided, ignore the engine-provided minimum value for the property and use this instead for display.
+  /// Note that this changes the min value displayed on the slider, but the user can continue sliding to reach lower
+  /// values.
+  final double? softMin;
+
+  /// If provided, ignore the engine-provided maximum value for the property and use this instead for display.
+  /// Note that this changes the max value displayed on the slider, but the user can continue sliding to reach higher
+  /// values.
+  final double? softMax;
+
+  /// If provided, ignore the engine-provided minimum value for the property and use this instead for clamping.
+  /// This changes both the max value displayed on the slider and prevents sliding past this value.
+  final double? hardMin;
+
+  /// If provided, ignore the engine-provided maximum value for the property and use this instead for clamping.
+  /// This changes both the max value displayed on the slider and prevents sliding past this value.
+  final double? hardMax;
+
+  /// If true, show the reset button.
+  final bool bShowResetButton;
+
+  /// The color to show in the filled portion of the bar (left of the value indicators).
+  final Color? fillColor;
+
+  @override
+  _UnrealDeltaSliderState createState() => _UnrealDeltaSliderState();
+}
+
+class _UnrealDeltaSliderState extends State<UnrealDeltaSlider> with UnrealWidgetStateMixin<UnrealDeltaSlider, double> {
+  @override
+  double? get overrideMin => widget.hardMin ?? engineMin;
+
+  @override
+  double? get overrideMax => widget.hardMax ?? engineMax;
+
+  @override
+  Widget build(BuildContext context) {
+    return TransientPreferenceBuilder(
+      preference: Provider.of<DeltaWidgetSettings>(context, listen: false).bIsInResetMode,
+      builder: (context, final bool bIsInResetMode) => DrivenDeltaSlider(
+        values: makeValueDataList(),
+        bShowResetButton: bIsInResetMode && widget.bShowResetButton,
+        bIsEnabled: !bIsInResetMode,
+        onChanged: handleOnChangedByUser,
+        onValuesSet: _onValuesSet,
+        onReset: handleOnResetByUser,
+        onInteractionFinished: endTransaction,
+        min: widget.softMin ?? widget.hardMin ?? engineMin,
+        max: widget.softMax ?? widget.hardMax ?? engineMax,
+        label: propertyLabel,
+        buildLabel: widget.buildLabel,
+        fillColor: widget.fillColor,
+      ),
+    );
+  }
+
+  /// Replace all existing values with a new list of values.
+  void _onValuesSet(List<double> newValues) {
+    modifyProperties(const SetOperation(), values: newValues, bIgnoreLimits: true);
+  }
+}
+
+/// A DeltaSlider whose values are controlled from outside of the widget itself.
+class DrivenDeltaSlider extends StatefulWidget {
+  const DrivenDeltaSlider({
+    Key? key,
+    required this.values,
+    required this.onChanged,
+    this.onValuesSet,
+    this.label = '',
+    this.buildLabel,
+    this.min = 0.0,
+    this.max = 1.0,
+    this.exponent = 1.0,
+    this.bShowResetButton = true,
+    this.bIgnoreSensitivity = false,
+    this.bIsEnabled = true,
+    this.defaultValue,
+    this.onReset,
+    this.onInteractionFinished,
+    this.fillColor,
+    this.baseSensitivity,
+  }) : super(key: key);
+
+  /// Values to display on the slider.
+  final List<DeltaSliderValueData?> values;
+
+  /// Name of the property this widget controls.
+  final String label;
+
+  /// If provided, call this function to build the widget displayed as the slider's label. Otherwise, the label will be
+  /// automatically generated.
+  final Widget Function(String name)? buildLabel;
+
+  /// Minimum value the user can select.
+  final double? min;
+
+  /// Maximum value the user can select.
+  final double? max;
+
+  /// Exponent value applied to the slider's linear position. Higher exponents make the slider's effective delta value
+  /// smaller at smaller slider values.
+  final double exponent;
+
+  /// If true, show the reset button.
+  final bool bShowResetButton;
+
+  /// If true, don't take sensitivity settings into account when calculating delta values.
+  final bool bIgnoreSensitivity;
+
+  /// If false, grey out the widget and don't accept inputs. Note that this doesn't affect the reset button.
+  final bool bIsEnabled;
+
+  /// Default value to reset to when the reset button is pressed.
+  final double? defaultValue;
+
+  /// Function called when the reset button is pressed. Called after the value has been reset.
+  final void Function()? onReset;
+
+  /// Function called when the user is done interacting with the widget.
+  final void Function()? onInteractionFinished;
+
+  /// Function called when the user changes the value of the slider.
+  /// Passes the amount by which each value changed.
+  final void Function(List<double>) onChanged;
+
+  /// Function called when the user directly sets the values of the slider.
+  /// Passes the new list of values.
+  /// If not set, these will be converted to deltas and passed to onChanged instead.
+  final void Function(List<double>)? onValuesSet;
+
+  /// The color to show in the filled portion of the bar (left of the value indicators).
+  final Color? fillColor;
+
+  /// How sensitive the slider should be when the user's sensitivity is set to 1.0.
+  final double? baseSensitivity;
+
+  @override
+  _DrivenDeltaSliderState createState() => _DrivenDeltaSliderState();
+}
+
+class _DrivenDeltaSliderState extends State<DrivenDeltaSlider> with DeltaWidgetStateMixin {
+  /// How long the min/max labels take to fade in/out.
+  static const _minMaxFadeDuration = Duration(milliseconds: 100);
+
+  bool _bIsBeingDragged = false;
+
+  @override
+  bool get bIgnoreSensitivity => widget.bIgnoreSensitivity;
+
+  @override
+  double get baseDeltaMultiplier => widget.baseSensitivity ?? super.baseDeltaMultiplier;
+
+  /// Get the slider's minimum value.
+  /// If the widget has a set minimum, return it. Otherwise, default to a minimum of 0, but extend the minimum to match
+  /// the lowest controlled value.
+  double get _min {
+    if (widget.min != null) {
+      return widget.min!;
+    }
+
+    double minValue = 0.0;
+    for (final DeltaSliderValueData? valueData in widget.values) {
+      if (valueData != null) {
+        minValue = math.min(valueData.value, minValue);
+      }
+    }
+
+    return minValue;
+  }
+
+  /// Get the slider's maximum value.
+  /// If the widget has a set maximum, return it. Otherwise, default to a maximum of 1, but extend the maximum to match
+  /// the highest controlled value.
+  double get _max {
+    if (widget.max != null) {
+      return widget.max!;
+    }
+
+    double maxValue = 1.0;
+    for (final DeltaSliderValueData? valueData in widget.values) {
+      if (valueData != null) {
+        maxValue = math.max(valueData.value, maxValue);
+      }
+    }
+
+    return maxValue;
+  }
+
+  /// Difference between the slider's min and max values.
+  double get _valueSpan => _max - _min;
+
+  /// Opacity to use for the slider's labels.
+  double get _labelOpacity => widget.bIsEnabled ? 1.0 : 0.4;
+
+  /// Style to use for the min/max labels
+  TextStyle get _minMaxLabelStyle => Theme.of(context).textTheme.labelMedium!.copyWith(
+        fontSize: 10,
+        color: UnrealColors.gray42.withOpacity(_labelOpacity),
+      );
+
+  /// The string used to indicate the slider's current value(s).
+  String? get _singleSharedValueString {
+    if (widget.values.isEmpty) {
+      return null;
+    }
+
+    final double? firstValue = widget.values[0]?.value;
+    if (firstValue == null) {
+      return null;
+    }
+
+    // If we have any values that differ, indicate it
+    for (int propertyIndex = 0; propertyIndex < widget.values.length; ++propertyIndex) {
+      if (widget.values[propertyIndex]?.value != firstValue) {
+        return AppLocalizations.of(context)!.mismatchedValuesLabel;
+      }
+    }
+
+    return _getStringForValue(firstValue);
+  }
+
+  @override
+  Widget build(BuildContext context) {
+    final String? valueString = _singleSharedValueString;
+
+    // Build the list of slider values to pass to the slider widget
+    final List<_DeltaSliderUIValue> uiValues = [];
+    for (final DeltaSliderValueData? valueData in widget.values) {
+      if (valueData != null) {
+        uiValues.add(
+            _DeltaSliderUIValue(valueData: valueData, normalizedValue: _getSliderPositionForValue(valueData.value)));
+      }
+    }
+
+    return DefaultTextStyle(
+      style: Theme.of(context).textTheme.labelMedium!.copyWith(
+            color: Theme.of(context).textTheme.labelMedium!.color!.withOpacity(_labelOpacity),
+          ),
+      child: Padding(
+        padding: const EdgeInsets.only(bottom: 5),
+        child: Row(
+          crossAxisAlignment: CrossAxisAlignment.start,
+          children: [
+            Expanded(
+              child: Column(children: [
+                Row(children: [
+                  // Main slider
+                  Expanded(
+                    child: GestureDetector(
+                      onHorizontalDragStart: _onDragStart,
+                      onHorizontalDragUpdate: _onDragUpdate,
+                      onHorizontalDragEnd: ((_) => _onDragEnd()),
+                      onHorizontalDragCancel: _onDragEnd,
+                      onDoubleTap: _showDirectEntryModal,
+                      behavior: HitTestBehavior.opaque,
+                      child: Column(
+                        children: [
+                          // Labels
+                          Padding(
+                            padding: const EdgeInsets.symmetric(horizontal: DeltaWidgetConstants.widgetOuterXPadding),
+                            child: ConstrainedBox(
+                              constraints: BoxConstraints.tightForFinite(),
+                              child: Row(
+                                mainAxisAlignment: MainAxisAlignment.spaceBetween,
+                                children: [
+                                  Expanded(child: _buildLabel(widget.label)),
+                                  if (valueString != null)
+                                    Padding(
+                                      padding: EdgeInsets.only(left: 6),
+                                      child: Text(valueString),
+                                    ),
+                                ],
+                              ),
+                            ),
+                          ),
+
+                          const SizedBox(height: 2),
+
+                          // Slider
+                          _DeltaSliderSlider(
+                            uiValues: uiValues,
+                            fillColor: widget.fillColor,
+                            bIsEnabled: widget.bIsEnabled,
+                          ),
+
+                          // Min/max labels
+                          AnimatedOpacity(
+                            opacity: _bIsBeingDragged ? 1 : 0,
+                            duration: _minMaxFadeDuration,
+                            child: Padding(
+                              padding: const EdgeInsets.only(
+                                left: DeltaWidgetConstants.widgetOuterXPadding,
+                                right: DeltaWidgetConstants.widgetOuterXPadding,
+                              ),
+                              child: Row(
+                                mainAxisAlignment: MainAxisAlignment.spaceBetween,
+                                children: [
+                                  Text(_getStringForValue(_min), style: _minMaxLabelStyle),
+                                  Text(_getStringForValue(_max), style: _minMaxLabelStyle),
+                                ],
+                              ),
+                            ),
+                          ),
+                        ],
+                      ),
+                    ),
+                  ),
+                ]),
+              ]),
+            ),
+            if (widget.bShowResetButton)
+              Padding(
+                padding: const EdgeInsets.only(right: DeltaWidgetConstants.widgetOuterXPadding),
+                child: Center(
+                  child: ResetValueButton(onPressed: widget.onReset),
+                ),
+              ),
+          ],
+        ),
+      ),
+    );
+  }
+
+  /// Normalize and clamp a value to the range of the slider.
+  double _normalizeValue(double value) => ((value - _min) / _valueSpan).clamp(0.0, 1.0);
+
+  /// Apply the slider's exponent to a normalized value.
+  double _exponentiateNormalizedValue(double normalized) {
+    if (widget.exponent == 1.0) {
+      return normalized;
+    }
+
+    return exponentiateValue(normalized, widget.exponent, true);
+  }
+
+  /// Convert a controlled value to the actual range of the slider, including normalization, clamping, and exponent.
+  double _getSliderPositionForValue(double value) => _exponentiateNormalizedValue(_normalizeValue(value));
+
+  /// Format a value on the slider as a string for display to the user.
+  String _getStringForValue(double value) {
+    return value.toStringAsFixed(2);
+  }
+
+  /// Called when the user starts dragging on the slider.
+  void _onDragStart(DragStartDetails details) {
+    setState(() {
+      _bIsBeingDragged = true;
+    });
+  }
+
+  /// Called when the user drags on the slider.
+  void _onDragUpdate(DragUpdateDetails details) {
+    if (!widget.bIsEnabled) {
+      return;
+    }
+
+    final double delta = details.primaryDelta! * _valueSpan * deltaMultiplier / (context.size?.width ?? 1.0);
+    final List<double> deltaValues = List.filled(widget.values.length, delta);
+
+    if (widget.exponent != 1.0) {
+      for (int valueIndex = 0; valueIndex < deltaValues.length; ++valueIndex) {
+        final double? currentValue = widget.values[valueIndex]?.value;
+        if (currentValue == null) {
+          continue;
+        }
+
+        final double normalizedValue = _normalizeValue(currentValue);
+        final double normalizedDelta = deltaValues[valueIndex] / _valueSpan;
+
+        // Delta is applied by user in exponential space, so find the delta in linear space
+        final double exponentialValue = _exponentiateNormalizedValue(normalizedValue);
+        final double exponentialResultValue = (exponentialValue + normalizedDelta).clamp(0, 1);
+        final double linearResultValue = exponentiateValue(exponentialResultValue, widget.exponent, false);
+        deltaValues[valueIndex] = linearResultValue - normalizedValue;
+      }
+    }
+
+    widget.onChanged(deltaValues);
+  }
+
+  /// Called when the user stops dragging on the slider.
+  void _onDragEnd() {
+    widget.onInteractionFinished?.call();
+    setState(() {
+      _bIsBeingDragged = false;
+    });
+  }
+
+  /// Build the widget to display as the slider's label.
+  Widget _buildLabel(String name) {
+    if (widget.buildLabel != null) {
+      return widget.buildLabel!(name);
+    }
+
+    return Text(name);
+  }
+
+  /// Show the modal dialog to directly enter a value via keyboard.
+  void _showDirectEntryModal() async {
+    if (widget.values.isEmpty || widget.values[0] == null) {
+      return;
+    }
+
+    final TextInputModalDialogResult<double>? result = await GenericModalDialogRoute.showDialog(
+      context: context,
+      builder: (context) => DoubleTextInputModalDialog(
+        title: AppLocalizations.of(context)!.sliderKeyboardInputModalTitle(widget.label),
+        initialValue: widget.values[0]!.value,
+        bShowResetButton: true,
+      ),
+    );
+
+    switch (result?.action) {
+      case TextInputModalDialogAction.apply:
+        _setAllValues(result!.value!);
+        break;
+
+      case TextInputModalDialogAction.reset:
+        widget.onReset?.call();
+        break;
+
+      case TextInputModalDialogAction.cancel:
+      case null:
+        break;
+    }
+  }
+
+  /// Change every value to a single, new value.
+  void _setAllValues(double newValue) {
+    _setValues(List.filled(widget.values.length, newValue));
+    widget.onInteractionFinished?.call();
+  }
+
+  /// Set the values to a list of new values.
+  void _setValues(List<double> newValues) {
+    assert(newValues.length == widget.values.length);
+
+    if (widget.onValuesSet != null) {
+      widget.onValuesSet!(newValues);
+    } else {
+      final List<double> deltas = [];
+
+      for (int valueIndex = 0; valueIndex < widget.values.length; ++valueIndex) {
+        final value = widget.values[valueIndex];
+        if (value == null) {
+          deltas.add(0);
+          continue;
+        }
+
+        deltas.add(newValues[valueIndex] - value.value);
+      }
+
+      widget.onChanged(deltas);
+    }
+  }
+}
+
+/// The interactive slider within the slider widget.
+class _DeltaSliderSlider extends StatelessWidget {
+  const _DeltaSliderSlider({
+    Key? key,
+    required this.uiValues,
+    required this.fillColor,
+    this.bIsEnabled = true,
+  }) : super(key: key);
+
+  final List<_DeltaSliderUIValue> uiValues;
+  final Color? fillColor;
+  final bool bIsEnabled;
+
+  @override
+  Widget build(BuildContext context) {
+    return SizedBox(
+      width: double.infinity,
+      height: _dotSize,
+      child: CustomPaint(
+        painter: _DeltaSliderSliderPainter(
+          context: context,
+          uiValues: uiValues,
+          fillColor: fillColor,
+          bIsEnabled: bIsEnabled,
+        ),
+      ),
+    );
+  }
+}
+
+/// Painter for the track and "handle" rings of a delta slider.
+class _DeltaSliderSliderPainter extends CustomPainter {
+  const _DeltaSliderSliderPainter({
+    required this.context,
+    required this.uiValues,
+    required this.fillColor,
+    this.bIsEnabled = true,
+  }) : super();
+
+  final BuildContext context;
+  final List<_DeltaSliderUIValue> uiValues;
+  final Color? fillColor;
+  final bool bIsEnabled;
+
+  @override
+  void paint(Canvas canvas, Size size) {
+    const double ringThickness = 2.8;
+    const double disabledOpacity = 0.4;
+
+    final Rect canvasRect = Offset.zero & size;
+    canvas.saveLayer(canvasRect, Paint());
+
+    // Determine positions of rings
+    final trackWidth = size.width - (_trackPadding * 2);
+    final ringSize = Size.square(_dotSize - ringThickness);
+    final double ringYOffset = (size.height - ringSize.height) / 2;
+
+    final Iterable<Rect> ringRects = uiValues.map((uiValue) {
+      final double ringX = _trackPadding + (uiValue.normalizedValue * trackWidth) - (ringSize.width / 2);
+      return Offset(ringX, ringYOffset) & ringSize;
+    });
+
+    // Draw the track colors
+    final double filledWidth = uiValues.fold(
+          0.0,
+          (double maxValue, _DeltaSliderUIValue uiValue) => math.max(maxValue, uiValue.normalizedValue),
+        ) *
+        size.width;
+
+    final Paint trackFillPaint = Paint()
+      ..style = PaintingStyle.fill
+      ..color = bIsEnabled
+          ? (fillColor ?? Theme.of(context).colorScheme.primary)
+          : UnrealColors.gray75.withOpacity(disabledOpacity);
+    canvas.drawRect(Offset.zero & Size(filledWidth, size.height), trackFillPaint);
+
+    final Paint trackEmptyPaint = Paint()
+      ..style = PaintingStyle.fill
+      ..color = UnrealColors.gray31.withOpacity(bIsEnabled ? 1.0 : disabledOpacity);
+    canvas.drawRect(Offset(filledWidth, 0) & Size(size.width - filledWidth, size.height), trackEmptyPaint);
+
+    canvas.saveLayer(canvasRect, Paint()..blendMode = BlendMode.dstIn);
+
+    // Create a mask in the shape of the track
+    final Paint trackPaint = Paint()
+      ..style = PaintingStyle.fill
+      ..color = Colors.white;
+
+    final trackRect = RRect.fromRectAndRadius(
+      Rect.fromLTWH(
+        _trackPadding,
+        (size.height - _trackHeight) / 2,
+        trackWidth,
+        _trackHeight,
+      ),
+      Radius.circular(_trackHeight),
+    );
+    canvas.drawRRect(trackRect, trackPaint);
+
+    // Erase the mask around the rings
+    final Paint ringMaskPaint = Paint()
+      ..style = PaintingStyle.fill
+      ..color = Colors.transparent
+      ..blendMode = BlendMode.clear;
+
+    const double trackRingPadding = 5.5;
+
+    for (final Rect ringRect in ringRects) {
+      canvas.drawRect(ringRect.inflate(trackRingPadding), ringMaskPaint);
+    }
+
+    // Restore the track color layer, showing only the part overlapping with the mask
+    canvas.restore();
+
+    // Draw rings
+    final Paint ringStrokePaint = Paint()
+      ..style = PaintingStyle.stroke
+      ..strokeWidth = ringThickness
+      ..color = UnrealColors.gray56.withOpacity(bIsEnabled ? 1.0 : disabledOpacity);
+
+    final Paint ringFillPaint = Paint()
+      ..style = PaintingStyle.fill
+      ..color = Colors.white
+      ..blendMode = BlendMode.dstOut;
+
+    for (final Rect ringRect in ringRects) {
+      canvas.drawArc(ringRect, 0, math.pi * 2, false, ringFillPaint);
+      canvas.drawArc(ringRect, 0, math.pi * 2, false, ringStrokePaint);
+    }
+
+    // Blend combined changes down to the canvas so transparency is preserved
+    canvas.restore();
+  }
+
+  @override
+  bool shouldRepaint(covariant CustomPainter oldDelegate) {
+    return oldDelegate != this;
+  }
+}
+
+/// Represents a value displayed on the delta slider in a format easier to use for UI purposes.
+class _DeltaSliderUIValue {
+  const _DeltaSliderUIValue({required this.valueData, required this.normalizedValue});
+
+  /// The value of the dot.
+  final DeltaSliderValueData valueData;
+
+  /// The value of the dot normalized to the slider's range.
+  final double normalizedValue;
+}
diff --git a/Engine/Extras/VirtualProduction/EpicStageApp/lib/widgets/elements/delta_widget_base.dart b/Engine/Extras/VirtualProduction/EpicStageApp/lib/widgets/elements/delta_widget_base.dart
new file mode 100644
--- /dev/null
+++ b/Engine/Extras/VirtualProduction/EpicStageApp/lib/widgets/elements/delta_widget_base.dart
@@ -0,0 +1,75 @@
+// Copyright Epic Games, Inc. All Rights Reserved.
+
+import 'package:flutter/material.dart';
+import 'package:flutter_gen/gen_l10n/app_localizations.dart';
+import 'package:provider/provider.dart';
+
+import '../../models/settings/delta_widget_settings.dart';
+import 'epic_icon_button.dart';
+
+class DeltaWidgetConstants {
+  /// Size of icon buttons used for delta widgets.
+  static const double deltaButtonIconSize = 32.0;
+
+  /// Default horizontal padding for delta widget buttons.
+  static const double defaultButtonXPadding = 6.0;
+
+  /// Size of an icon button including the default padding on both sides.
+  static const double deltaButtonSize = deltaButtonIconSize + defaultButtonXPadding * 2;
+
+  /// Padding between the edge of a widget's container and its visible components.
+  static const double widgetOuterXPadding = 10.0;
+
+  /// Horizontal offset between a widget's label and the widget control itself.
+  static const double widgetInnerXPadding = 2.0;
+}
+
+/// A button to reset the value controlled by a widget.
+class ResetValueButton extends StatelessWidget {
+  const ResetValueButton(
+      {this.onPressed, this.bEnabled = true, this.xPadding = DeltaWidgetConstants.defaultButtonXPadding, Key? key})
+      : super(key: key);
+
+  /// Callback for when the button is pressed.
+  final Function()? onPressed;
+
+  /// Whether the button is enabled. If false, events won't be called and the app_icon will be greyed out.
+  final bool bEnabled;
+
+  /// Padding along the x-axis on either side.
+  final double? xPadding;
+
+  @override
+  Widget build(BuildContext context) {
+    return EpicIconButton(
+      onPressed: onPressed,
+      iconPath: 'assets/images/icons/reset.svg',
+      tooltipMessage: AppLocalizations.of(context)!.propertyResetButtonTooltip,
+      iconSize: 24,
+      buttonSize: const Size(32, 32),
+    );
+  }
+}
+
+/// Mixin that adds state for any widgets that use delta-based controls.
+mixin DeltaWidgetStateMixin<T extends StatefulWidget> on State<T> {
+  /// If true, ignore the sensitivity settings and just use the base delta multiplier.
+  bool get bIgnoreSensitivity => false;
+
+  /// Base scroll rate of a value (generally multiplied by the gesture's delta movement divided the widget's size).
+  /// Override this to change the base scroll rate for a widget.
+  double get baseDeltaMultiplier => 0.8;
+
+  /// The effective multiplier, combining the base rate and any other scalars.
+  double get deltaMultiplier =>
+      baseDeltaMultiplier *
+      (bIgnoreSensitivity ? 1.0 : Provider.of<DeltaWidgetSettings>(context, listen: false).sensitivity.getValue());
+}
+
+/// Data needed to display a control representing a single object's property on a delta-based widget.
+class DeltaWidgetValueData<T> {
+  const DeltaWidgetValueData({required this.value});
+
+  /// The value the dot represents.
+  final T value;
+}
diff --git a/Engine/Extras/VirtualProduction/EpicStageApp/lib/widgets/elements/dropdown_button.dart b/Engine/Extras/VirtualProduction/EpicStageApp/lib/widgets/elements/dropdown_button.dart
new file mode 100644
--- /dev/null
+++ b/Engine/Extras/VirtualProduction/EpicStageApp/lib/widgets/elements/dropdown_button.dart
@@ -0,0 +1,375 @@
+// Copyright Epic Games, Inc. All Rights Reserved.
+
+import 'package:flutter/material.dart';
+import 'package:provider/provider.dart';
+
+import '../../models/navigator_keys.dart';
+import '../../models/property_modify_operations.dart';
+import '../../models/settings/delta_widget_settings.dart';
+import '../../utilities/transient_preference.dart';
+import '../../utilities/unreal_colors.dart';
+import 'asset_icon.dart';
+import 'delta_widget_base.dart';
+import 'dropdown_list_menu.dart';
+import 'list_menu.dart';
+import 'unreal_widget_base.dart';
+
+/// Possible states for a [ModalTextDropdownButton].
+enum ModalDropdownButtonState {
+  /// The modal is closed.
+  closed,
+
+  /// The modal is open and displaying under the button.
+  openOnTop,
+
+  /// The modal is open and displaying above the button.
+  openOnBottom,
+}
+
+/// A button that uses a builder function to open a dropdown list modal when pressed.
+class ModalDropdownButton<ItemType> extends StatefulWidget {
+  const ModalDropdownButton({
+    Key? key,
+    required this.buttonBuilder,
+    required this.menuBuilder,
+    this.openRectExtraSpace = EdgeInsets.zero,
+    this.bStretchDropdown = false,
+    this.tooltipMessage,
+  }) : super(key: key);
+
+  /// Function that builds the button based on the dropdown's [state].
+  final Widget Function(BuildContext context, ModalDropdownButtonState state) buttonBuilder;
+
+  /// Function that builds the dropdown menu.
+  /// Receives an [originTabBuilder] function which will build the menu's origin tab based on whether the tab [isOnTop].
+  final Widget Function(BuildContext context, Widget Function(BuildContext, bool isOnTop) originTabBuilder) menuBuilder;
+
+  /// How much space to add around the original button's bounding area when building its open version.
+  final EdgeInsets openRectExtraSpace;
+
+  /// Whether to stretch the dropdown to match the size of the button.
+  final bool bStretchDropdown;
+
+  /// An optional tooltip message to show when the button is hovered or long pressed.
+  final String? tooltipMessage;
+
+  @override
+  State<ModalDropdownButton<ItemType>> createState() => _ModalDropdownButtonState<ItemType>();
+}
+
+class _ModalDropdownButtonState<ItemType> extends State<ModalDropdownButton<ItemType>> {
+  /// Key used to identify the widget from which the dropdown originates.
+  final GlobalKey _dropdownOriginator = GlobalKey();
+
+  @override
+  Widget build(BuildContext context) {
+    final button = Semantics(
+      button: true,
+      child: MouseRegion(
+        cursor: MaterialStateMouseCursor.clickable,
+        child: GestureDetector(
+          key: _dropdownOriginator,
+          onTap: _onTap,
+          child: widget.buttonBuilder(context, ModalDropdownButtonState.closed),
+        ),
+      ),
+    );
+
+    return (widget.tooltipMessage != null)
+        ? Tooltip(
+            child: button,
+            message: widget.tooltipMessage,
+          )
+        : button;
+  }
+
+  /// Called when this is tapped.
+  void _onTap() {
+    final RenderObject? widgetRenderObject = _dropdownOriginator.currentContext?.findRenderObject();
+    final RenderBox? widgetBox = (widgetRenderObject != null) ? (widgetRenderObject as RenderBox) : null;
+
+    if (widgetBox == null) {
+      throw Exception('Failed to find widget render box for $this');
+    }
+
+    final RenderObject rootRenderObject = rootNavigatorKey.currentContext!.findRenderObject()!;
+
+    final Size rectSize = Size(
+      widgetBox.size.width + widget.openRectExtraSpace.horizontal,
+      widgetBox.size.height + widget.openRectExtraSpace.vertical,
+    );
+
+    final Offset globalPos = widgetBox.localToGlobal(
+      Offset(-widget.openRectExtraSpace.left, -widget.openRectExtraSpace.bottom),
+      ancestor: rootRenderObject,
+    );
+
+    final Rect widgetRect = Rect.fromLTWH(globalPos.dx, globalPos.dy, rectSize.width, rectSize.height);
+
+    DropDownListMenu.showAtRect(
+      pivotRect: widgetRect,
+      bStretch: widget.bStretchDropdown,
+      builder: (context) => widget.menuBuilder(
+        context,
+        (context, bIsOnTop) => ConstrainedBox(
+          constraints: BoxConstraints(
+            maxWidth: widgetRect.width,
+            maxHeight: widgetRect.height,
+          ),
+          child: widget.buttonBuilder(
+            context,
+            bIsOnTop ? ModalDropdownButtonState.openOnTop : ModalDropdownButtonState.openOnBottom,
+          ),
+        ),
+      ),
+    );
+  }
+}
+
+/// A button that opens an automatically built dropdown list modal containing text-only options when pressed.
+class ModalTextDropdownButton<ItemType> extends StatelessWidget {
+  const ModalTextDropdownButton({
+    Key? key,
+    required this.builder,
+    required this.items,
+    required this.onChanged,
+    this.value,
+    this.makeItemName,
+    this.openRectExtraSpace = EdgeInsets.zero,
+    this.bStretchDropdown = false,
+  }) : super(key: key);
+
+  /// Function that builds the button based on the dropdown's state.
+  final Widget Function(BuildContext context, ModalDropdownButtonState state) builder;
+
+  /// The currently selected value.
+  final ItemType? value;
+
+  /// The list of selectable items in the order in which they should be displayed.
+  final List<ItemType> items;
+
+  /// Function called when a new value is selected.
+  final void Function(ItemType newValue) onChanged;
+
+  /// Optional function to convert items to a display string. If not provided, the items will be converted implicitly.
+  final String Function(ItemType item)? makeItemName;
+
+  /// How much space to add around the original button's bounding area when building its open version.
+  final EdgeInsets openRectExtraSpace;
+
+  /// Whether to stretch the dropdown to match the size of the button.
+  final bool bStretchDropdown;
+
+  @override
+  Widget build(BuildContext context) {
+    return ModalDropdownButton(
+      openRectExtraSpace: openRectExtraSpace,
+      bStretchDropdown: bStretchDropdown,
+      buttonBuilder: builder,
+      menuBuilder: (context, originTabBuilder) => DropDownListMenu(
+        originTabBuilder: originTabBuilder,
+        children: [
+          for (ItemType item in items)
+            ListMenuSimpleItem(
+              title: _makeItemName(item),
+              onTap: () => _onValueSelected(context, item),
+            ),
+        ],
+      ),
+    );
+  }
+
+  /// Convert an item to a display name, using the user-provided function if possible.
+  String _makeItemName(ItemType item) {
+    if (makeItemName != null) {
+      return makeItemName!(item);
+    }
+
+    return item.toString();
+  }
+
+  /// Called when an item is selected.
+  void _onValueSelected(BuildContext context, ItemType value) {
+    final navigator = Navigator.of(context, rootNavigator: true);
+    navigator.pop();
+
+    onChanged(value);
+  }
+}
+
+/// A dropdown button styled as a large selector containing text.
+class DropdownSelector<ItemType> extends StatelessWidget {
+  const DropdownSelector({
+    Key? key,
+    required this.value,
+    required this.items,
+    required this.onChanged,
+    this.hint = '',
+    this.makeItemName,
+    this.bIsEnabled = true,
+    this.borderRadius = const Radius.circular(4),
+  }) : super(key: key);
+
+  /// The currently selected value.
+  final ItemType? value;
+
+  /// The list of selectable items in the order in which they should be displayed.
+  final List<ItemType> items;
+
+  /// Function called when a new value is selected.
+  final void Function(ItemType newValue) onChanged;
+
+  /// Optional function to convert items to a display string. If not provided, the items will be converted implicitly.
+  final String Function(ItemType item)? makeItemName;
+
+  /// If false, disable the controls and grey this out.
+  final bool bIsEnabled;
+
+  /// Text to show when [value] is null.
+  final String hint;
+
+  /// The border radius of the button.
+  final Radius borderRadius;
+
+  /// The display name to show inside the dropdown button.
+  String get _displayName {
+    if (value == null) {
+      return hint;
+    }
+
+    if (makeItemName != null) {
+      return makeItemName!(value!);
+    }
+
+    return value.toString();
+  }
+
+  @override
+  Widget build(BuildContext context) {
+    final Color textColor = bIsEnabled ? UnrealColors.white : UnrealColors.gray56;
+
+    final builder = (BuildContext context, ModalDropdownButtonState state) => Container(
+          height: 36,
+          decoration: BoxDecoration(
+            color: bIsEnabled ? Theme.of(context).colorScheme.background : UnrealColors.gray22,
+            borderRadius: BorderRadius.vertical(
+              top: borderRadius,
+              bottom: state == ModalDropdownButtonState.closed ? borderRadius : Radius.zero,
+            ),
+          ),
+          padding: EdgeInsets.symmetric(horizontal: 12),
+          child: Row(
+            mainAxisAlignment: MainAxisAlignment.spaceBetween,
+            crossAxisAlignment: CrossAxisAlignment.center,
+            children: [
+              Text(
+                _displayName,
+                style: Theme.of(context).textTheme.bodyMedium!.copyWith(
+                      color: textColor,
+                      letterSpacing: 0.5,
+                    ),
+              ),
+              SizedBox(width: 8),
+              AssetIcon(
+                path: 'assets/images/icons/chevron_down.svg',
+                size: 16,
+                color: textColor,
+              ),
+            ],
+          ),
+        );
+
+    if (bIsEnabled) {
+      return ModalTextDropdownButton(
+        bStretchDropdown: true,
+        value: value,
+        items: items,
+        onChanged: onChanged,
+        makeItemName: makeItemName,
+        builder: builder,
+      );
+    } else {
+      return builder(context, ModalDropdownButtonState.closed);
+    }
+  }
+}
+
+/// A dropdown that controls one or more enum properties in Unreal Engine.
+class UnrealDropdownSelector extends UnrealWidget {
+  const UnrealDropdownSelector({
+    super.key,
+    super.overrideName,
+    super.enableProperties,
+    required super.unrealProperties,
+  });
+
+  @override
+  _UnrealDropdownButtonState createState() => _UnrealDropdownButtonState();
+}
+
+class _UnrealDropdownButtonState extends State<UnrealDropdownSelector>
+    with UnrealWidgetStateMixin<UnrealDropdownSelector, String> {
+  @override
+  PropertyModifyOperation get modifyOperation => const SetOperation();
+
+  @override
+  Widget build(BuildContext context) {
+    final SingleSharedValue<String> sharedValue = getSingleSharedValue();
+    final TextStyle textStyle = Theme.of(context).textTheme.labelMedium!;
+
+    return Padding(
+      padding: EdgeInsets.only(bottom: 8),
+      child: TransientPreferenceBuilder(
+        preference: Provider.of<DeltaWidgetSettings>(context, listen: false).bIsInResetMode,
+        builder: (context, final bool bIsInResetMode) => Row(
+          children: [
+            Expanded(
+              child: Padding(
+                padding: EdgeInsets.symmetric(horizontal: DeltaWidgetConstants.widgetOuterXPadding),
+                child: Column(
+                  mainAxisAlignment: MainAxisAlignment.start,
+                  crossAxisAlignment: CrossAxisAlignment.start,
+                  children: [
+                    Text(
+                      propertyLabel,
+                      style: textStyle.copyWith(
+                        color: textStyle.color!.withOpacity(bIsInResetMode ? 0.4 : 1.0),
+                      ),
+                    ),
+                    Padding(
+                      padding: const EdgeInsets.only(
+                        left: DeltaWidgetConstants.widgetInnerXPadding,
+                        right: DeltaWidgetConstants.widgetInnerXPadding,
+                        top: 4,
+                      ),
+                      child: DropdownSelector<String>(
+                        items: propertyEnumValues,
+                        value: sharedValue.value,
+                        onChanged: _onChanged,
+                        hint: sharedValue.bHasMultipleValues ? 'Multiple values' : '',
+                        bIsEnabled: !bIsInResetMode,
+                      ),
+                    ),
+                  ],
+                ),
+              ),
+            ),
+            if (bIsInResetMode)
+              Padding(
+                padding: const EdgeInsets.only(right: DeltaWidgetConstants.widgetOuterXPadding),
+                child: ResetValueButton(onPressed: handleOnResetByUser),
+              ),
+          ],
+        ),
+      ),
+    );
+  }
+
+  void _onChanged(dynamic newValue) {
+    assert(newValue is String);
+    assert(propertyEnumValues.contains(newValue));
+
+    handleOnChangedByUser(List.filled(widget.unrealProperties.length, newValue));
+    endTransaction();
+  }
+}
