Branch snapshot for CL 27138787

diff --git a/Engine/Extras/VirtualProduction/EpicStageApp/lib/widgets/elements/dropdown_list_menu.dart b/Engine/Extras/VirtualProduction/EpicStageApp/lib/widgets/elements/dropdown_list_menu.dart
new file mode 100644
--- /dev/null
+++ b/Engine/Extras/VirtualProduction/EpicStageApp/lib/widgets/elements/dropdown_list_menu.dart
@@ -0,0 +1,422 @@
+// Copyright Epic Games, Inc. All Rights Reserved.
+
+import 'dart:math';
+
+import 'package:flutter/material.dart';
+import 'package:logging/logging.dart';
+import 'package:provider/provider.dart';
+
+import '../../models/navigator_keys.dart';
+import 'layout/measurable.dart';
+import 'list_menu.dart';
+
+final _log = Logger('DropdownListMenu');
+
+/// Function type used to create a [DropDownListMenu].
+typedef DropDownListMenuBuilder = Widget Function(BuildContext context);
+
+/// Holds data about how a dropdown list should be placed on-screen.
+class DropDownListPlacement {
+  const DropDownListPlacement({
+    required this.pivotRect,
+    required this.bStretch,
+  });
+
+  /// The rectangle next to which the dropdown should be placed.
+  final Rect pivotRect;
+
+  /// Whether the menu should be stretched to fit the pivot rectangle along the opposite axis to [pivotAxis].
+  final bool bStretch;
+}
+
+/// A modal drop-down list menu that can be placed at an arbitrary position on the screen.
+class DropDownListMenu extends StatefulWidget {
+  const DropDownListMenu({
+    Key? widgetKey,
+    required this.children,
+    this.backgroundColor,
+    this.originTabBuilder,
+    this.bCenterOriginTabOnPivot = false,
+    this.minHeight = 0,
+  }) : super(key: widgetKey);
+
+  /// Items in the list.
+  final List<Widget> children;
+
+  /// Optional background color for the list.
+  final Color? backgroundColor;
+
+  /// Optional builder for an extra widget to display on top of/underneath the menu.
+  /// This will be positioned to line up with the menu's origin point.
+  final Widget Function(BuildContext context, bool bIsOnTop)? originTabBuilder;
+
+  /// If true, position the menu with the origin tab built by [originTabBuilder] centered on the center of the pivot.
+  final bool bCenterOriginTabOnPivot;
+
+  /// The list's minimum height. If list would extend off the screen, it will be constrained to this height at minimum
+  /// and add a scroll bar.
+  final double minHeight;
+
+  /// Show a drop-down menu originating from the given [pivotRect] and created using the [builder]. It will be
+  /// automatically placed at the side of the rectangle with the most space.
+  static void showAtRect({
+    required Rect pivotRect,
+    required DropDownListMenuBuilder builder,
+    bool bStretch = false,
+    void Function()? onPopped,
+  }) {
+    final route = _DropDownListMenuRoute(
+      placement: DropDownListPlacement(
+        bStretch: bStretch,
+        pivotRect: pivotRect,
+      ),
+      builder: builder,
+    );
+
+    if (onPopped != null) {
+      route.popped.then((_) => onPopped());
+    }
+
+    Navigator.of(rootNavigatorKey.currentContext!, rootNavigator: true).push<void>(route);
+  }
+
+  /// Show a drop-down menu at the given [pivotPosition] and created using the [builder].
+  /// If provided, the menu will be placed at a distance of [pivotRadius] from the given position.
+  static void showAtPosition({
+    required Offset pivotPosition,
+    required DropDownListMenuBuilder builder,
+    bool bStretch = false,
+    Size? pivotSize,
+    void Function()? onPopped,
+  }) {
+    pivotSize = pivotSize ?? const Size(32, 64);
+
+    showAtRect(
+      pivotRect: Rect.fromCenter(
+        center: pivotPosition,
+        width: pivotSize.width,
+        height: pivotSize.height,
+      ),
+      bStretch: bStretch,
+      builder: builder,
+      onPopped: onPopped,
+    );
+  }
+
+  /// Show a drop-down menu next to the widget with the given [widgetKey] and created using the [builder].
+  /// If provided, the widget's bounds will be inflated by [rectPadding] for the purpose of choosing the menu's
+  /// position.
+  static void showAtWidget({
+    required GlobalKey widgetKey,
+    required DropDownListMenuBuilder builder,
+    bool bStretch = false,
+    double rectPadding = 0,
+    void Function()? onPopped,
+  }) {
+    final RenderObject? widgetRenderObject = widgetKey.currentContext?.findRenderObject();
+    final RenderBox? widgetBox = (widgetRenderObject != null) ? (widgetRenderObject as RenderBox) : null;
+    if (widgetBox == null) {
+      _log.severe('No widget provided for drop-down menu location');
+      return;
+    }
+
+    final RenderObject rootRenderObject = rootNavigatorKey.currentContext!.findRenderObject()!;
+    Rect widgetRect = widgetBox.localToGlobal(Offset.zero, ancestor: rootRenderObject) & widgetBox.size;
+    widgetRect = widgetRect.inflate(rectPadding);
+
+    showAtRect(
+      pivotRect: widgetRect,
+      bStretch: bStretch,
+      builder: builder,
+      onPopped: onPopped,
+    );
+  }
+
+  @override
+  State<StatefulWidget> createState() => _DropDownListMenuState();
+}
+
+class _DropDownListMenuState extends State<DropDownListMenu> {
+  Size? _menuSize;
+  Size? _menuSizeWithScrollbar;
+  Size? _originTabSize;
+
+  @override
+  Widget build(BuildContext context) {
+    return Consumer<DropDownListPlacement>(builder: (context, dropdownData, _) {
+      if (_menuSize == null) {
+        // Build once so we can determine the menu's size. Use Align here so the menu doesn't expand to fill the entire
+        // screen. This assumes that the menu won't change size on the next build (e.g. if its edgeDirection changes).
+        return Offstage(
+          child: Align(
+            child: Stack(
+              children: [
+                if (widget.originTabBuilder != null)
+                  Measurable(
+                    child: widget.originTabBuilder!(context, true),
+                    onMeasured: _onOriginTabMeasured,
+                  ),
+                Measurable(
+                  child: ListMenu(children: widget.children),
+                  onMeasured: _onMenuMeasured,
+                ),
+                Measurable(
+                  child: ListMenu(children: widget.children, bIsScrollable: true),
+                  onMeasured: _onMenuWithScrollBarMeasured,
+                ),
+              ],
+            ),
+          ),
+        );
+      }
+
+      if (_menuSize == null || _menuSizeWithScrollbar == null || _originTabSize == null) {
+        // Still waiting for initial build to be measured
+        return SizedBox();
+      }
+
+      final Size screenSize = MediaQuery.of(context).size;
+      final EdgeInsets screenPadding = MediaQuery.of(context).viewPadding;
+
+      late final Rect pivotRect;
+      late final Offset initialPosition;
+
+      if (widget.bCenterOriginTabOnPivot) {
+        pivotRect = dropdownData.pivotRect.center & Size.zero;
+        initialPosition = pivotRect.center +
+            Offset(
+              -_originTabSize!.width / 2,
+              _originTabSize!.height / 2,
+            );
+      } else {
+        pivotRect = dropdownData.pivotRect;
+        initialPosition = pivotRect.bottomLeft;
+      }
+
+      // Adjust Y position if we would extend off the bottom of the screen
+      // Max Y coordinate of menu's bottom
+      final double bottomLimit = screenSize.height - screenPadding.bottom;
+
+      // Actual Y coordinate of menu's bottom
+      final double menuBottom = initialPosition.dy + widget.minHeight;
+
+      // The height of the menu, which we will apply constraints to as necessary to fit on-screen
+      double menuHeight = _menuSize!.height;
+
+      bool bIsOriginTop = true;
+      double yPosition = initialPosition.dy;
+
+      if (menuBottom > bottomLimit) {
+        // Display the menu above the pivot instead of below
+        bIsOriginTop = false;
+
+        if (widget.bCenterOriginTabOnPivot) {
+          yPosition -= _menuSize!.height + _originTabSize!.height;
+        } else {
+          yPosition -= _menuSize!.height + pivotRect.height;
+        }
+
+        if (yPosition < 0) {
+          // Shrink menu to fit on screen
+          menuHeight += yPosition;
+          yPosition = 0;
+        }
+      } else {
+        // Clamp height in bounds of screen
+        menuHeight = menuHeight.clamp(
+          widget.minHeight,
+          bottomLimit - yPosition,
+        );
+      }
+
+      // Adjust Y position to accommodate origin tab
+      if (_originTabSize != null && bIsOriginTop) {
+        yPosition -= _originTabSize!.height;
+      }
+
+      // Minimum menu width so rounded border lines up with edge
+      double minMenuWidth = 0;
+      if (_originTabSize != null && !dropdownData.bStretch) {
+        minMenuWidth = _originTabSize!.width + ListMenu.borderRadius.x;
+      }
+
+      final bool bWillScroll = menuHeight < _menuSize!.height;
+      final double menuWidth = max(minMenuWidth, bWillScroll ? _menuSizeWithScrollbar!.width : _menuSize!.width);
+
+      // Adjust X position if we would extend off the right of the screen
+      bool bIsOriginLeft = true;
+      double xPosition = initialPosition.dx;
+
+      if (!dropdownData.bStretch) {
+        if (xPosition < screenPadding.left) {
+          xPosition = screenPadding.left;
+        } else {
+          // Push menu to the left if the dropdown would cross the right edge of the safe area
+          final double rightLimit = screenSize.width - screenPadding.left;
+          final double menuRight = xPosition + menuWidth;
+
+          if (menuRight > rightLimit) {
+            // Display menu to left of pivot instead of right
+            if (widget.bCenterOriginTabOnPivot) {
+              xPosition -= menuWidth - _originTabSize!.width;
+            } else {
+              xPosition -= menuWidth - pivotRect.width;
+            }
+            bIsOriginLeft = false;
+
+            // Move the menu left to fit on-screen
+            final double newRight = xPosition + menuWidth;
+            if (newRight > rightLimit) {
+              xPosition -= newRight - rightLimit;
+            }
+          }
+        }
+      }
+
+      late final ListMenuEdgeDirection edgeDirection;
+      if (bIsOriginTop) {
+        if (dropdownData.bStretch) {
+          edgeDirection = ListMenuEdgeDirection.top;
+        } else if (bIsOriginLeft) {
+          edgeDirection = ListMenuEdgeDirection.topLeft;
+        } else {
+          edgeDirection = ListMenuEdgeDirection.topRight;
+        }
+      } else {
+        if (dropdownData.bStretch) {
+          edgeDirection = ListMenuEdgeDirection.bottom;
+        } else if (bIsOriginLeft) {
+          edgeDirection = ListMenuEdgeDirection.bottomLeft;
+        } else {
+          edgeDirection = ListMenuEdgeDirection.bottomRight;
+        }
+      }
+
+      // Create the origin tab if necessary
+      late final Widget? originTab;
+      if (widget.originTabBuilder != null) {
+        // Ignore the pointer so tapping the tab also closes the dropdown
+        originTab = IgnorePointer(
+          child: widget.originTabBuilder!(context, bIsOriginTop),
+        );
+      } else {
+        originTab = null;
+      }
+
+      // Create the contents of the dropdown menu.
+      Widget listWidget = ListMenu(
+        minWidth: dropdownData.bStretch ? pivotRect.width : 0,
+        maxWidth: dropdownData.bStretch ? pivotRect.width : double.infinity,
+        children: widget.children,
+        backgroundColor: widget.backgroundColor ?? Theme.of(context).colorScheme.surfaceTint,
+        edgeDirection: edgeDirection,
+        bIsScrollable: bWillScroll,
+      );
+
+      // Wrap with scrolling view
+      if (bWillScroll) {
+        listWidget = SizedBox(
+          height: menuHeight,
+          child: listWidget,
+        );
+      }
+
+      return CustomSingleChildLayout(
+        delegate: _DropDownListMenuLayout(position: Offset(xPosition, yPosition)),
+        child: Align(
+          alignment: Alignment.topLeft,
+          child: Column(
+            crossAxisAlignment: bIsOriginLeft ? CrossAxisAlignment.start : CrossAxisAlignment.end,
+            children: [
+              if (bIsOriginTop && originTab != null) originTab,
+              ConstrainedBox(
+                constraints: BoxConstraints(minWidth: minMenuWidth),
+                child: listWidget,
+              ),
+              if (!bIsOriginTop && originTab != null) originTab,
+            ],
+          ),
+        ),
+      );
+    });
+  }
+
+  void _onMenuMeasured(Size size) {
+    if (_menuSize != size) {
+      setState(() {
+        _menuSize = size;
+      });
+    }
+  }
+
+  void _onMenuWithScrollBarMeasured(Size size) {
+    if (_menuSizeWithScrollbar != size) {
+      setState(() {
+        _menuSizeWithScrollbar = size;
+      });
+    }
+  }
+
+  void _onOriginTabMeasured(Size size) {
+    if (_originTabSize != size) {
+      setState(() {
+        _originTabSize = size;
+      });
+    }
+  }
+}
+
+/// A route that displays the drop-down menu popup and positions it.
+class _DropDownListMenuRoute<T> extends PopupRoute<T> {
+  _DropDownListMenuRoute({
+    required this.placement,
+    required this.builder,
+  });
+
+  final DropDownListPlacement placement;
+  final DropDownListMenuBuilder builder;
+
+  @override
+  Color? get barrierColor => Colors.black38;
+
+  @override
+  bool get barrierDismissible => true;
+
+  @override
+  String? get barrierLabel => 'Dismiss';
+
+  @override
+  Duration get transitionDuration => Duration(milliseconds: 200);
+
+  @override
+  Duration get reverseTransitionDuration => Duration(milliseconds: 100);
+
+  @override
+  Widget buildPage(BuildContext context, Animation<double> animation, Animation<double> secondaryAnimation) {
+    return Builder(
+      builder: (BuildContext context) {
+        return Provider(
+          create: (context) => placement,
+          child: builder(context),
+        );
+      },
+    );
+  }
+}
+
+/// Layout that places the dropdown within the screen.
+class _DropDownListMenuLayout<T> extends SingleChildLayoutDelegate {
+  _DropDownListMenuLayout({required this.position});
+
+  final Offset position;
+
+  @override
+  Offset getPositionForChild(Size size, Size childSize) {
+    return position;
+  }
+
+  @override
+  bool shouldRelayout(covariant SingleChildLayoutDelegate oldDelegate) {
+    return oldDelegate != this;
+  }
+}
diff --git a/Engine/Extras/VirtualProduction/EpicStageApp/lib/widgets/elements/dropdown_text.dart b/Engine/Extras/VirtualProduction/EpicStageApp/lib/widgets/elements/dropdown_text.dart
new file mode 100644
--- /dev/null
+++ b/Engine/Extras/VirtualProduction/EpicStageApp/lib/widgets/elements/dropdown_text.dart
@@ -0,0 +1,264 @@
+// Copyright Epic Games, Inc. All Rights Reserved.
+
+import 'package:flutter/material.dart';
+import 'package:flutter_gen/gen_l10n/app_localizations.dart';
+import 'package:provider/provider.dart';
+
+import '../../models/property_modify_operations.dart';
+import '../../models/settings/delta_widget_settings.dart';
+import '../../utilities/transient_preference.dart';
+import '../../utilities/unreal_colors.dart';
+import 'asset_icon.dart';
+import 'dropdown_button.dart';
+import 'unreal_widget_base.dart';
+
+/// Padding used inside the origin tab when the dropdown is open.
+const EdgeInsets _originTabPadding = EdgeInsets.only(
+  left: 8,
+  right: 4,
+  top: 12,
+  bottom: 12,
+);
+
+/// Text which acts as a dropdown, but otherwise remains inline.
+/// Controls an enum value in Unreal Engine.
+class UnrealDropdownText extends UnrealWidget {
+  const UnrealDropdownText({
+    super.key,
+    super.overrideName,
+    super.enableProperties,
+    required super.unrealProperties,
+    this.makeItemName,
+  });
+
+  /// Optional function to convert items to a display string. If not provided, the items will be converted implicitly.
+  final String Function(String enumValueName)? makeItemName;
+
+  @override
+  State<UnrealDropdownText> createState() => _UnrealDropdownTextState();
+}
+
+class _UnrealDropdownTextState extends State<UnrealDropdownText>
+    with UnrealWidgetStateMixin<UnrealDropdownText, String> {
+  @override
+  Widget build(BuildContext context) {
+    final SingleSharedValue<String> sharedValue = getSingleSharedValue();
+
+    return TransientPreferenceBuilder(
+      preference: Provider.of<DeltaWidgetSettings>(context, listen: false).bIsInResetMode,
+      builder: (context, final bool bIsInResetMode) => DrivenDropdownText<SingleSharedValue<String>>(
+        value: sharedValue,
+        items: propertyEnumValues
+            .map((enumValue) => SingleSharedValue(
+                  value: enumValue,
+                  bHasMultipleValues: false,
+                ))
+            .toList(growable: false),
+        makeItemName: _makeItemName,
+        onChanged: _onChanged,
+        bIsEnabled: !bIsInResetMode,
+      ),
+    );
+  }
+
+  @override
+  PropertyModifyOperation get modifyOperation => const SetOperation();
+
+  String? _makeItemName(SingleSharedValue<String> value) {
+    if (value.bHasMultipleValues) {
+      return AppLocalizations.of(context)!.mismatchedValuesLabel;
+    }
+
+    if (value.value == null) {
+      return null;
+    }
+
+    return value.value!;
+  }
+
+  void _onChanged(dynamic newValue) {
+    assert(newValue is SingleSharedValue<String>);
+
+    final sharedValue = newValue as SingleSharedValue<String>;
+
+    assert(propertyEnumValues.contains(sharedValue.value));
+
+    handleOnChangedByUser(List.filled(widget.unrealProperties.length, sharedValue.value!));
+    endTransaction();
+  }
+}
+
+/// Text which acts as a dropdown, but otherwise remains inline.
+class DrivenDropdownText<ItemType> extends StatelessWidget {
+  const DrivenDropdownText({
+    Key? key,
+    required this.value,
+    required this.items,
+    required this.onChanged,
+    this.makeItemName,
+    this.bIsEnabled = true,
+  }) : super(key: key);
+
+  /// The currently selected value.
+  final ItemType value;
+
+  /// The list of selectable items in the order in which they should be displayed.
+  final List<ItemType> items;
+
+  /// Function called when a new value is selected.
+  final void Function(ItemType newValue) onChanged;
+
+  /// Optional function to convert items to a display string. If not provided, the items will be converted implicitly.
+  final String? Function(ItemType item)? makeItemName;
+
+  /// If false, grey out the widget and disable its controls.
+  final bool bIsEnabled;
+
+  @override
+  Widget build(BuildContext context) {
+    // Align the row based on text directionality, respecting the text direction.
+    late final MainAxisAlignment alignment;
+    switch (DefaultTextStyle.of(context).textAlign) {
+      case TextAlign.start:
+        alignment = MainAxisAlignment.start;
+        break;
+
+      case TextAlign.end:
+        alignment = MainAxisAlignment.end;
+        break;
+
+      case TextAlign.center:
+        alignment = MainAxisAlignment.center;
+        break;
+
+      case TextAlign.left:
+        alignment = (Directionality.of(context) == TextDirection.ltr) ? MainAxisAlignment.start : MainAxisAlignment.end;
+        break;
+
+      case TextAlign.right:
+        alignment = (Directionality.of(context) == TextDirection.ltr) ? MainAxisAlignment.end : MainAxisAlignment.start;
+        break;
+
+      default:
+        alignment = MainAxisAlignment.start;
+    }
+
+    if (!bIsEnabled) {
+      return DefaultTextStyle(
+        style: DefaultTextStyle.of(context).style.copyWith(
+              color: DefaultTextStyle.of(context).style.color!.withOpacity(0.4),
+            ),
+        child: _DrivenDropdownTextInner(
+          text: _makeItemName(value),
+          alignment: alignment,
+          bHideChevron: true,
+        ),
+      );
+    }
+
+    return ModalTextDropdownButton(
+      value: value,
+      items: items,
+      onChanged: onChanged,
+      makeItemName: _makeItemName,
+      openRectExtraSpace: _originTabPadding,
+      builder: (BuildContext innerContext, ModalDropdownButtonState state) {
+        final String name = _makeItemName(value);
+
+        if (state == ModalDropdownButtonState.closed) {
+          return _DrivenDropdownTextInner(
+            text: name,
+            alignment: alignment,
+          );
+        }
+
+        final bool bIsOnTop = state == ModalDropdownButtonState.openOnTop;
+
+        return Container(
+          decoration: BoxDecoration(
+            color: UnrealColors.gray22,
+            borderRadius: BorderRadius.vertical(
+              top: bIsOnTop ? Radius.circular(8) : Radius.zero,
+              bottom: !bIsOnTop ? Radius.circular(8) : Radius.zero,
+            ),
+          ),
+          padding: _originTabPadding,
+          child: DefaultTextStyle(
+            style: DefaultTextStyle.of(context).style,
+            child: _DrivenDropdownTextInner(
+              text: name,
+              alignment: MainAxisAlignment.start,
+              color: UnrealColors.white,
+            ),
+          ),
+        );
+      },
+    );
+  }
+
+  /// Convert an item to a display name, using the user-provided function if possible.
+  String _makeItemName(ItemType item) {
+    if (item == null) {
+      return '...';
+    }
+
+    if (makeItemName != null) {
+      return makeItemName!(item) ?? '...';
+    }
+
+    return item.toString();
+  }
+}
+
+/// The inner contents of a [DrivenDropdownText], showing the text and the chevron icon.
+class _DrivenDropdownTextInner extends StatelessWidget {
+  const _DrivenDropdownTextInner({
+    Key? key,
+    required this.text,
+    required this.alignment,
+    this.color,
+    this.bHideChevron = false,
+  }) : super(key: key);
+
+  final String text;
+  final MainAxisAlignment alignment;
+  final Color? color;
+  final bool bHideChevron;
+
+  @override
+  Widget build(BuildContext context) {
+    TextStyle textStyle = DefaultTextStyle.of(context).style;
+    if (color != null) {
+      textStyle = textStyle.copyWith(color: color);
+    }
+
+    return Row(
+      mainAxisSize: MainAxisSize.min,
+      mainAxisAlignment: alignment,
+      children: [
+        Flexible(
+          child: Text(
+            text,
+            overflow: TextOverflow.ellipsis,
+            softWrap: false,
+            style: textStyle,
+          ),
+        ),
+        if (!bHideChevron)
+          SizedOverflowBox(
+            size: const Size(24, 15),
+            alignment: Alignment.center,
+            child: SizedBox(
+              width: 24,
+              child: AssetIcon(
+                path: 'assets/images/icons/chevron_down.svg',
+                size: 24,
+                color: textStyle.color,
+                fit: BoxFit.cover,
+              ),
+            ),
+          ),
+      ],
+    );
+  }
+}
diff --git a/Engine/Extras/VirtualProduction/EpicStageApp/lib/widgets/elements/empty_placeholder.dart b/Engine/Extras/VirtualProduction/EpicStageApp/lib/widgets/elements/empty_placeholder.dart
new file mode 100644
--- /dev/null
+++ b/Engine/Extras/VirtualProduction/EpicStageApp/lib/widgets/elements/empty_placeholder.dart
@@ -0,0 +1,54 @@
+// Copyright Epic Games, Inc. All Rights Reserved.
+
+import 'dart:ui';
+
+import 'package:flutter/material.dart';
+
+import '../../../../utilities/unreal_colors.dart';
+import 'asset_icon.dart';
+
+/// Widget for displaying text and an optional button when a pane has nothing else to show.
+class EmptyPlaceholder extends StatelessWidget {
+  const EmptyPlaceholder({Key? key, required this.message, this.button}) : super(key: key);
+
+  /// The main message to display.
+  final String message;
+
+  /// An optional button widget to show under the message.
+  final Widget? button;
+
+  @override
+  Widget build(BuildContext context) {
+    return Padding(
+      padding: EdgeInsets.all(16),
+      child: Column(
+        mainAxisAlignment: MainAxisAlignment.center,
+        children: [
+          AssetIcon(
+            path: 'assets/images/icons/unreal_u.svg',
+            size: 64,
+            color: UnrealColors.gray75,
+          ),
+          const SizedBox(height: 32),
+          Text(
+            message,
+            style: TextStyle(
+              color: UnrealColors.gray56,
+              fontStyle: FontStyle.italic,
+              fontVariations: [FontVariation('wght', 400)],
+              fontSize: 20,
+              letterSpacing: 0.5,
+              height: 1,
+            ),
+            textAlign: TextAlign.center,
+          ),
+          if (button != null)
+            Padding(
+              padding: EdgeInsets.only(top: 32),
+              child: button,
+            ),
+        ],
+      ),
+    );
+  }
+}
diff --git a/Engine/Extras/VirtualProduction/EpicStageApp/lib/widgets/elements/epic_checkbox.dart b/Engine/Extras/VirtualProduction/EpicStageApp/lib/widgets/elements/epic_checkbox.dart
new file mode 100644
--- /dev/null
+++ b/Engine/Extras/VirtualProduction/EpicStageApp/lib/widgets/elements/epic_checkbox.dart
@@ -0,0 +1,53 @@
+// Copyright Epic Games, Inc. All Rights Reserved.
+
+import 'package:flutter/material.dart';
+
+import 'asset_icon.dart';
+
+/// A checkbox button that matches the Epic app style.
+class EpicCheckbox extends StatelessWidget {
+  const EpicCheckbox({
+    Key? key,
+    required this.bChecked,
+    this.label,
+    this.onPressed,
+  }) : super(key: key);
+
+  /// If true, show a checkmark in the box.
+  final bool bChecked;
+
+  /// Label text to show next to the checkbox.
+  final String? label;
+
+  /// Callback function when this is pressed.
+  final Function()? onPressed;
+
+  @override
+  Widget build(BuildContext context) {
+    return Semantics(
+      button: true,
+      child: GestureDetector(
+        onTap: onPressed,
+        child: Row(
+          crossAxisAlignment: CrossAxisAlignment.center,
+          children: [
+            AssetIcon(
+              size: 20,
+              path: bChecked
+                  ? 'assets/images/icons/checkbox_opaque_checked.svg'
+                  : 'assets/images/icons/checkbox_unchecked.svg',
+            ),
+            if (label != null)
+              Padding(
+                padding: EdgeInsets.symmetric(horizontal: 8),
+                child: Text(
+                  label!,
+                  style: Theme.of(context).textTheme.headlineSmall,
+                ),
+              ),
+          ],
+        ),
+      ),
+    );
+  }
+}
diff --git a/Engine/Extras/VirtualProduction/EpicStageApp/lib/widgets/elements/epic_icon_button.dart b/Engine/Extras/VirtualProduction/EpicStageApp/lib/widgets/elements/epic_icon_button.dart
new file mode 100644
--- /dev/null
+++ b/Engine/Extras/VirtualProduction/EpicStageApp/lib/widgets/elements/epic_icon_button.dart
@@ -0,0 +1,343 @@
+// Copyright Epic Games, Inc. All Rights Reserved.
+
+import 'dart:ui';
+
+import 'package:flutter/material.dart';
+
+import '../../utilities/unreal_colors.dart';
+import 'asset_icon.dart';
+
+/// A standard-sized button with arbitrary contents.
+class EpicGenericButton extends StatelessWidget {
+  const EpicGenericButton({
+    Key? key,
+    required this.child,
+    this.tooltipMessage,
+    this.onPressed,
+    this.size,
+    this.backgroundSize,
+    this.backgroundShape,
+    this.bIsToggledOn = false,
+    this.color,
+    this.activeBackgroundColor,
+    this.bIsVisualOnly = false,
+    this.onLongPressed,
+  }) : super(key: key);
+
+  final Widget child;
+
+  /// An optional tooltip message to show when the button is hovered or long pressed.
+  final String? tooltipMessage;
+
+  /// The function to call when the button is pressed. If not provided and [bIsVisualOnly] if false, the button is
+  /// considered inactive and will change color accordingly.
+  final Function()? onPressed;
+
+  /// The function to call when the button is long pressed
+  final Function()? onLongPressed;
+
+  /// The size of the button.
+  final Size? size;
+
+  /// The size of the button's background when visible.
+  final Size? backgroundSize;
+
+  /// The shape of the button's background when visible.
+  final BoxShape? backgroundShape;
+
+  /// If true, highlight the button using the [activeBackgroundColor].
+  final bool bIsToggledOn;
+
+  /// The color of the button's contents.
+  final Color? color;
+
+  /// The background color of the button when it's toggled on.
+  final Color? activeBackgroundColor;
+
+  /// If true, leave out interactive behaviour and only include the button's visuals.
+  /// This will disable the tooltip message.
+  final bool bIsVisualOnly;
+
+  @override
+  Widget build(BuildContext context) {
+    Color? backgroundColor = null;
+
+    if (bIsToggledOn) {
+      if (onPressed == null && !bIsVisualOnly && onLongPressed == null) {
+        // Disabled color
+        backgroundColor = UnrealColors.gray42;
+      } else {
+        // Toggled on color
+        backgroundColor = activeBackgroundColor ?? Theme.of(context).colorScheme.primary;
+      }
+    }
+
+    Widget contents = Container(
+      decoration: BoxDecoration(
+        color: backgroundColor,
+        shape: backgroundShape ?? BoxShape.rectangle,
+      ),
+      width: backgroundSize?.width ?? size?.width ?? 48,
+      height: backgroundSize?.height ?? size?.height ?? 48,
+      child: DefaultTextStyle(
+        style: Theme.of(context).textTheme.bodyMedium!.copyWith(color: color ?? UnrealColors.white),
+        child: Center(
+          child: child,
+        ),
+      ),
+    );
+
+    // Add button semantics
+    if (!bIsVisualOnly) {
+      contents = Semantics(
+        button: true,
+        child: MouseRegion(
+          cursor: MaterialStateMouseCursor.clickable,
+          child: GestureDetector(
+            child: contents,
+            onTap: onPressed,
+            behavior: HitTestBehavior.opaque,
+            onLongPress: onLongPressed,
+          ),
+        ),
+      );
+    }
+
+    if (tooltipMessage == null) {
+      return contents;
+    }
+
+    return Tooltip(
+      message: tooltipMessage,
+      child: contents,
+    );
+  }
+}
+
+/// A standard-sized button with an icon.
+class EpicIconButton extends StatelessWidget {
+  const EpicIconButton({
+    Key? key,
+    required this.iconPath,
+    this.tooltipMessage,
+    this.onPressed,
+    this.bIsToggledOn = false,
+    this.iconSize = 24,
+    this.buttonSize,
+    this.backgroundSize,
+    this.backgroundShape,
+    this.color,
+    this.activeColor,
+    this.activeBackgroundColor,
+    this.bIsVisualOnly = false,
+    this.onLongPressed,
+  }) : super(key: key);
+
+  /// The path of the image file to show inside the button.
+  final String iconPath;
+
+  /// An optional tooltip message to show when the button is hovered or long pressed.
+  final String? tooltipMessage;
+
+  /// The function to call when the button is pressed. If not provided and [bIsVisualOnly] if false, the button is
+  /// considered inactive and will change color accordingly.
+  final Function()? onPressed;
+
+  /// The function to call when the button is Long pressed. If not provided and
+  /// [bIsVisualOnly] if false, the button is
+  /// considered inactive and will change color accordingly.
+  final Function()? onLongPressed;
+
+  /// If true, highlight the button using the [activeBackgroundColor].
+  final bool bIsToggledOn;
+
+  /// The size of the icon (both width and height).
+  final double iconSize;
+
+  /// The size of the button's pressable area.
+  final Size? buttonSize;
+
+  /// The size of the button's background when visible.
+  final Size? backgroundSize;
+
+  /// The shape of the button's background when visible.
+  final BoxShape? backgroundShape;
+
+  /// The color of the button's contents.
+  final Color? color;
+
+  /// The color of the button's when the button is toggled on.
+  final Color? activeColor;
+
+  /// The background color of the button when it's toggled on.
+  final Color? activeBackgroundColor;
+
+  /// If true, leave out button semantics and only include the button's visuals.
+  /// This will disable the tooltip message.
+  final bool bIsVisualOnly;
+
+  @override
+  Widget build(BuildContext context) {
+    Color? iconColor = color;
+
+    if (bIsToggledOn) {
+      iconColor = activeColor;
+    } else if (onPressed == null && onLongPressed == null && !bIsVisualOnly) {
+      iconColor = UnrealColors.gray42;
+    }
+
+    return EpicGenericButton(
+      child: AssetIcon(
+        path: iconPath,
+        color: iconColor,
+        size: iconSize,
+      ),
+      tooltipMessage: tooltipMessage,
+      onPressed: onPressed,
+      onLongPressed: onLongPressed,
+      bIsToggledOn: bIsToggledOn,
+      color: color,
+      size: buttonSize,
+      backgroundSize: backgroundSize,
+      backgroundShape: backgroundShape,
+      activeBackgroundColor: activeBackgroundColor,
+      bIsVisualOnly: bIsVisualOnly,
+    );
+  }
+}
+
+/// A standard-sized rounded button containing text.
+class EpicLozengeButton extends StatelessWidget {
+  const EpicLozengeButton({
+    Key? key,
+    required this.label,
+    this.width,
+    this.color,
+    this.onPressed,
+  }) : super(key: key);
+
+  /// The button's label.
+  final String label;
+
+  /// If set, force the button to be exactly this width.
+  final double? width;
+
+  /// The color of the button's contents.
+  /// Defaults to the theme's primary color. If this is transparent, the text will be colored darker.
+  final Color? color;
+
+  /// The function to call when the button is pressed. If not provided and [bIsVisualOnly] if false, the button is
+  /// considered inactive and will change color accordingly.
+  final Function()? onPressed;
+
+  @override
+  Widget build(BuildContext context) {
+    final ThemeData theme = Theme.of(context);
+
+    return Semantics(
+      button: true,
+      child: MouseRegion(
+        cursor: MaterialStateMouseCursor.clickable,
+        child: GestureDetector(
+          onTap: onPressed,
+          behavior: HitTestBehavior.opaque,
+          child: Container(
+            width: width,
+            height: 36,
+            padding: EdgeInsets.symmetric(vertical: 10, horizontal: 28),
+            decoration: BoxDecoration(
+              color: color ?? theme.colorScheme.primary,
+              borderRadius: BorderRadius.circular(20),
+            ),
+            child: Center(
+              child: Text(
+                label,
+                style: theme.textTheme.headlineSmall!.copyWith(
+                  color: color?.alpha == 0 ? theme.colorScheme.onSurface : theme.colorScheme.onPrimary,
+                ),
+              ),
+            ),
+          ),
+        ),
+      ),
+    );
+  }
+}
+
+/// A standard-sized short, wide button with text and an icon.
+class EpicWideButton extends StatelessWidget {
+  const EpicWideButton({
+    Key? key,
+    required this.text,
+    required this.iconPath,
+    this.onPressed,
+    this.color,
+    this.iconColor,
+    this.textColor,
+    this.border,
+  }) : super(key: key);
+
+  /// The text to show within the button.
+  final String text;
+
+  /// The path of the asset to use for the icon within the button.
+  final String iconPath;
+
+  /// Callback function for when the user presses the button.
+  final Function()? onPressed;
+
+  /// An optional color for the button's background. If not provided, uses the surfaceTint of the theme's color scheme.
+  final Color? color;
+
+  /// An optional color for the button's icon.
+  final Color? iconColor;
+
+  /// An optional color for the button's text. If not provided, uses white.
+  final Color? textColor;
+
+  /// An optional border to show around the button. If not provided, no border is shown.
+  final Border? border;
+
+  @override
+  Widget build(BuildContext context) {
+    return Semantics(
+      button: true,
+      child: MouseRegion(
+        cursor: MaterialStateMouseCursor.clickable,
+        child: GestureDetector(
+          onTap: onPressed,
+          behavior: HitTestBehavior.opaque,
+          child: Container(
+            height: 32,
+            decoration: BoxDecoration(
+              color: color ?? UnrealColors.gray22,
+              border: border,
+              borderRadius: BorderRadius.circular(4),
+            ),
+            padding: EdgeInsets.only(left: 6, right: 8),
+            child: Row(
+              mainAxisSize: MainAxisSize.min,
+              mainAxisAlignment: MainAxisAlignment.center,
+              crossAxisAlignment: CrossAxisAlignment.center,
+              children: [
+                AssetIcon(
+                  path: iconPath,
+                  size: 24,
+                  color: iconColor,
+                ),
+                const SizedBox(width: 6),
+                Text(
+                  text,
+                  style: Theme.of(context).textTheme.bodyMedium!.copyWith(
+                    color: textColor ?? UnrealColors.white,
+                    fontVariations: [FontVariation('wght', 600)],
+                  ),
+                ),
+              ],
+            ),
+          ),
+        ),
+      ),
+    );
+  }
+}
diff --git a/Engine/Extras/VirtualProduction/EpicStageApp/lib/widgets/elements/epic_list_view.dart b/Engine/Extras/VirtualProduction/EpicStageApp/lib/widgets/elements/epic_list_view.dart
new file mode 100644
--- /dev/null
+++ b/Engine/Extras/VirtualProduction/EpicStageApp/lib/widgets/elements/epic_list_view.dart
@@ -0,0 +1,156 @@
+// Copyright Epic Games, Inc. All Rights Reserved.
+
+import 'package:flutter/material.dart';
+
+import '../../utilities/unreal_colors.dart';
+
+typedef IndexedWidgetBuilder(BuildContext context, int index);
+
+/// Core ListView for the app, wraps a [EpicScrollBar] + [ListView.builder].
+class EpicListView extends StatefulWidget {
+  /// defines extent for the children of the ListView.
+  final double? itemExtent;
+
+  /// Number of items to be rendered on the list.
+  final int itemCount;
+
+  /// Additional padding around the main scrollable area.
+  final EdgeInsets padding;
+
+  /// Builders callback for rendering each item in the current [context] and for the given [index]
+  final IndexedWidgetBuilder itemBuilder;
+
+  /// whether to show scrollbar or not.
+  final bool? bShowScrollbar;
+
+  /// Color of the scrollbar thumb.
+  final Color? scrollbarThumbColor;
+
+  EpicListView({
+    Key? key,
+    required this.itemBuilder,
+    required this.itemCount,
+    this.itemExtent,
+    this.padding = EdgeInsets.zero,
+    this.bShowScrollbar,
+    this.scrollbarThumbColor,
+  }) : super(key: key);
+
+  @override
+  _EpicListViewState createState() => _EpicListViewState();
+}
+
+class _EpicListViewState extends State<EpicListView> {
+  late _ScrollControllerWithMixin scrollController;
+
+  @override
+  void initState() {
+    super.initState();
+    scrollController = _ScrollControllerWithMixin();
+  }
+
+  @override
+  void didUpdateWidget(EpicListView oldWidget) {
+    super.didUpdateWidget(oldWidget);
+  }
+
+  @override
+  Widget build(BuildContext context) {
+    WidgetsBinding.instance.addPostFrameCallback((_) {
+      setState(() {});
+    });
+    return EpicScrollBar(
+      controller: scrollController,
+      padding: widget.padding.add(
+        EdgeInsets.only(
+          right: (scrollController.bIsScrollable ?? false) ? 8 : 0,
+        ),
+      ) as EdgeInsets,
+      bEnabled: widget.bShowScrollbar ?? true,
+      thumbColor: widget.scrollbarThumbColor,
+      child: ScrollConfiguration(
+        behavior: ScrollConfiguration.of(context).copyWith(scrollbars: false),
+        child: ListView.builder(
+          controller: scrollController,
+          itemExtent: widget.itemExtent,
+          itemCount: widget.itemCount,
+          padding: EdgeInsets.zero,
+          itemBuilder: (BuildContext context, int index) {
+            return Padding(
+              padding: EdgeInsets.only(
+                right: (scrollController.bIsScrollable ?? false) ? 15.0 : 0.0,
+                top: (index > 0) ? 2 : 0,
+              ),
+              child: widget.itemBuilder(context, index),
+            );
+          },
+        ),
+      ),
+    );
+  }
+}
+
+/// Mixin for reusable scroll controller with easy access to internal methods.
+mixin ScrollControllerMixin on ScrollController {
+  /// whether scrollview the controller is attached to is scrollable or not.
+  bool? get bIsScrollable {
+    if (this.hasClients && this.position.hasContentDimensions) {
+      return this.hasClients && (this.position.minScrollExtent < this.position.maxScrollExtent);
+    }
+    return false;
+  }
+}
+
+/// ScrollController for [EpicListView] with [ScrollControllerMixin] for easy usage.
+class _ScrollControllerWithMixin extends ScrollController with ScrollControllerMixin {}
+
+/// Styled theme scroll bar.
+class EpicScrollBar extends StatelessWidget {
+  const EpicScrollBar({
+    Key? key,
+    required this.child,
+    required this.controller,
+    this.bEnabled = true,
+    this.notification = defaultScrollNotificationPredicate,
+    this.thumbColor,
+    this.padding = EdgeInsets.zero,
+  }) : super(key: key);
+
+  /// Whether styled scrollbar is enabled or not.
+  final bool bEnabled;
+
+  /// Child widget to be wrapped with scrollbar possibly be a list.
+  final Widget child;
+
+  /// Creating a binding with the [child] scroll view.
+  final ScrollController controller;
+
+  /// Color for the scrollbar handle or thumb region.
+  final Color? thumbColor;
+
+  /// Scrollbar padding.
+  final EdgeInsets padding;
+
+  /// Tell the scrollbar which ScrollView to control and respond to as [child] could have multiple scroll views.
+  final ScrollNotificationPredicate notification;
+
+  @override
+  Widget build(BuildContext context) {
+    return bEnabled
+        ? Padding(
+            padding: padding,
+            child: RawScrollbar(
+              padding: EdgeInsets.only(top: 4, bottom: 4),
+              controller: controller,
+              radius: Theme.of(context).scrollbarTheme.radius,
+              thickness: 8,
+              trackVisibility: false,
+              thumbVisibility: true,
+              thumbColor: thumbColor ?? UnrealColors.gray22,
+              notificationPredicate: notification,
+              child: child,
+            ),
+          )
+        : child;
+  }
+}
diff --git a/Engine/Extras/VirtualProduction/EpicStageApp/lib/widgets/elements/epic_scroll_view.dart b/Engine/Extras/VirtualProduction/EpicStageApp/lib/widgets/elements/epic_scroll_view.dart
new file mode 100644
--- /dev/null
+++ b/Engine/Extras/VirtualProduction/EpicStageApp/lib/widgets/elements/epic_scroll_view.dart
@@ -0,0 +1,192 @@
+// Copyright Epic Games, Inc. All Rights Reserved.
+
+import 'dart:ui';
+
+import 'package:flutter/material.dart';
+
+import '../../utilities/preloadable_shader.dart';
+import '../../utilities/unreal_colors.dart';
+import 'epic_list_view.dart';
+
+const double _scrollbarPadding = 24;
+
+/// A custom-styled scrolling view which includes a fixed scrollbar, padding to fit it, and a gradient fading items
+/// near the edges of the scroll area.
+class EpicScrollView extends StatefulWidget {
+  const EpicScrollView({
+    Key? key,
+    required this.child,
+  }) : super(key: key);
+
+  final Widget child;
+
+  static Future preloadShaders() => _EpicScrollViewShader.load();
+  static void unloadShaders() => _EpicScrollViewShader.unload();
+
+  @override
+  State<EpicScrollView> createState() => _EpicScrollViewState();
+}
+
+class _EpicScrollViewState extends State<EpicScrollView> with SingleTickerProviderStateMixin {
+  final ScrollController _controller = ScrollController();
+  late final _EpicScrollViewShader shader;
+
+  @override
+  void initState() {
+    super.initState();
+    shader = _EpicScrollViewShader(
+      scrollController: _controller,
+      vsync: this,
+    );
+  }
+
+  @override
+  void dispose() {
+    shader.dispose();
+    super.dispose();
+  }
+
+  @override
+  Widget build(BuildContext context) {
+    return MediaQuery.removePadding(
+      context: context,
+      removeBottom: true,
+      removeTop: true,
+      child: ShaderMask(
+        shaderCallback: shader.shaderMaskCallback,
+        child: EpicScrollBar(
+          controller: _controller,
+          padding: EdgeInsets.only(right: 8, top: 8, bottom: 8),
+          thumbColor: UnrealColors.gray31,
+          child: ScrollConfiguration(
+            behavior: ScrollConfiguration.of(context).copyWith(scrollbars: false),
+            child: SingleChildScrollView(
+              controller: _controller,
+              padding: EdgeInsets.only(
+                right: _scrollbarPadding,
+              ),
+              child: widget.child,
+            ),
+          ),
+        ),
+        blendMode: BlendMode.modulate,
+      ),
+    );
+  }
+
+  @override
+  void reassemble() {
+    super.reassemble();
+
+    shader.reload();
+  }
+}
+
+/// Indices for each uniform value in the shader used by [_EpicScrollViewShader].
+enum _ScrollViewShaderVars {
+  uScrollbarPadding,
+  uStrength,
+  uScrollOffset,
+  uMaxScrollOffset,
+  uSizeX,
+  uSizeY,
+}
+
+/// Shader that fades the top and bottom of the scroll view when there's more to scroll in that direction.
+class _EpicScrollViewShader extends PreloadableShader {
+  _EpicScrollViewShader({
+    required this.scrollController,
+    required TickerProvider vsync,
+  }) : _fadeInController = AnimationController(
+          value: 0,
+          upperBound: 1,
+          duration: const Duration(milliseconds: 100),
+          vsync: vsync,
+        ) {
+    scrollController.addListener(_updateOffset);
+    _fadeInController.addListener(_updateStrength);
+
+    // Scroll controller won't be ready immediately, so try updating and fading in the effect after build
+    WidgetsBinding.instance.addPostFrameCallback((_) {
+      _updateOffset();
+      _fadeInEffectIfReady();
+    });
+  }
+
+  static final PreloadableShaderProgram _program = PreloadableShaderProgram('assets/shaders/scroll_view_fade.frag');
+
+  /// The scroll controller for the corresponding scroll view.
+  final ScrollController scrollController;
+
+  /// The animation controller used to fade in the effect when ready.
+  final AnimationController _fadeInController;
+
+  /// Whether this has been disposed.
+  bool _bIsDisposed = false;
+
+  @override
+  PreloadableShaderProgram get program => _program;
+
+  /// Load or re-load the shader program.
+  static Future<FragmentProgram> load() => _program.load();
+
+  /// Unload the shader program.
+  static void unload() => _program.unload();
+
+  /// Callback function passed to a [ShaderMask] which updates the bounding rectangle.
+  Shader shaderMaskCallback(Rect rect) {
+    shader!.setFloat(_ScrollViewShaderVars.uSizeX.index, rect.width);
+    shader!.setFloat(_ScrollViewShaderVars.uSizeY.index, rect.height);
+
+    return shader!;
+  }
+
+  @override
+  void dispose() {
+    super.dispose();
+
+    if (_bIsDisposed) {
+      return;
+    }
+
+    scrollController.removeListener(_updateOffset);
+    _fadeInController.dispose();
+    _bIsDisposed = true;
+  }
+
+  @override
+  void initUniforms() {
+    shader!.setFloat(_ScrollViewShaderVars.uScrollbarPadding.index, _scrollbarPadding);
+    _updateStrength();
+    _updateOffset();
+  }
+
+  /// Fade in the effect if it hasn't already and the scroll controller is ready to be accessed.
+  void _fadeInEffectIfReady() {
+    if (_bIsDisposed ||
+        _fadeInController.isCompleted ||
+        _fadeInController.isAnimating ||
+        scrollController.positions.isEmpty) {
+      return;
+    }
+
+    _fadeInController.animateTo(1.0);
+  }
+
+  /// Update the effect's strength based on the fade in animation.
+  void _updateStrength() {
+    shader!.setFloat(_ScrollViewShaderVars.uStrength.index, _fadeInController.value);
+  }
+
+  /// Update the shader's properties based on the [scrollController]'s position.
+  void _updateOffset() {
+    if (scrollController.positions.isEmpty) {
+      return;
+    }
+
+    _fadeInEffectIfReady();
+
+    shader!.setFloat(_ScrollViewShaderVars.uScrollOffset.index, scrollController.offset);
+    shader!.setFloat(_ScrollViewShaderVars.uMaxScrollOffset.index, scrollController.position.maxScrollExtent);
+  }
+}
diff --git a/Engine/Extras/VirtualProduction/EpicStageApp/lib/widgets/elements/floating_map_preview.dart b/Engine/Extras/VirtualProduction/EpicStageApp/lib/widgets/elements/floating_map_preview.dart
new file mode 100644
--- /dev/null
+++ b/Engine/Extras/VirtualProduction/EpicStageApp/lib/widgets/elements/floating_map_preview.dart
@@ -0,0 +1,686 @@
+// Copyright Epic Games, Inc. All Rights Reserved.
+
+import 'dart:math' as math;
+
+import 'package:flutter/material.dart';
+import 'package:flutter/physics.dart';
+import 'package:flutter/scheduler.dart';
+
+import '../../models/settings/floating_map_settings.dart';
+import '../../utilities/preferences_bundle.dart';
+import 'asset_icon.dart';
+import 'lightcard_map.dart';
+
+/// Sides of the screen on which the preview can be docked.
+enum _DockSide {
+  left,
+  right,
+}
+
+/// Possible states for the floating preview in terms of docking.
+enum _DockState {
+  /// The preview is on-screen and no dock tabs are visible.
+  undocked,
+
+  /// The preview has entered the docking range and is revealing its tab.
+  docking,
+
+  /// The preview has revealed its docking tab. It may not actually be off-screen yet.
+  docked,
+
+  /// The preview has exited the docking range and is hiding its tab.
+  undocking,
+}
+
+/// The smallest dimension of the container will be multiplied by this value to determine the height of the preview.
+const double _previewSizeScale = 0.25;
+
+/// The height of the preview will be multiplied by this value to determine its width.
+const double _previewAspectRatio = 16.0 / 9.0;
+
+/// How far to inset the preview from the edge of the screen.
+const double _edgeInset = 8.0;
+
+/// The size of the tab shown when the preview is docked.
+const _dockTabSize = Size(32, 80);
+
+/// A floating preview of the stage map that the user can drag to any corner of the screen and dock to the side.
+class FloatingMapPreview extends StatefulWidget {
+  const FloatingMapPreview({Key? key}) : super(key: key);
+
+  @override
+  State<FloatingMapPreview> createState() => _FloatingMapPreviewState();
+}
+
+class _FloatingMapPreviewState extends State<FloatingMapPreview> with SingleTickerProviderStateMixin {
+  /// Physics settings for the spring simulation when the user is dragging the preview.
+  static const _draggedSpringDescription = SpringDescription(mass: 40, stiffness: 1, damping: 1.2);
+
+  /// Physics settings for the spring simulation when the user releases the preview.
+  static const _releasedSpringDescription = SpringDescription(mass: 40, stiffness: 1, damping: 1);
+
+  /// Physics settings for the spring simulation when the preview is being docked.
+  static const _dockedSpringDescription = SpringDescription(mass: 60, stiffness: 1, damping: 1);
+
+  /// When this portion of the preview's width is offscreen, dock it.
+  static const double _dockThreshold = 1 / 3;
+
+  /// Which side the preview was last docked on.
+  _DockSide _dockSide = _DockSide.left;
+
+  /// State of the preview window's docking.
+  _DockState _dockState = _DockState.undocked;
+
+  /// The position that the preview's top-left corner should move towards with spring motion.
+  Offset _targetPreviewPosition = Offset.zero;
+
+  /// The current position of the preview's top-left corner.
+  Offset _previewPosition = Offset.zero;
+
+  /// The current velocity of the preview.
+  Offset _previewVelocity = Offset.zero;
+
+  /// The user's last drag position in global space.
+  Offset _lastDragPosition = Offset.zero;
+
+  /// The last known size of the containing render box.
+  Size _containerSize = Size.zero;
+
+  /// The elapsed ticker time in seconds when the springs were last updated.
+  double _lastSpringTime = 0;
+
+  /// The last value of [_lastSpringTime] when the springs' target positions were moved.
+  double _lastSpringMoveTime = 0;
+
+  /// Spring controlling the X axis.
+  SpringSimulation? _xSpring;
+
+  /// Spring controlling the Y axis.
+  SpringSimulation? _ySpring;
+
+  /// Ticker used to update springs.
+  late final Ticker _ticker;
+
+  /// Settings for this preview map.
+  late final FloatingMapSettings _settings;
+
+  /// The size of the preview to display given the containing widget's size.
+  Size get _previewSize {
+    final Size screenSize = MediaQuery.of(context).size;
+    final double previewHeight = math.min(screenSize.width, screenSize.height) * _previewSizeScale;
+
+    return Size(previewHeight * _previewAspectRatio, previewHeight);
+  }
+
+  /// True if the preview is currently docked or in the process of docking.
+  bool get _bIsDockedOrDocking => _dockState == _DockState.docking || _dockState == _DockState.docked;
+
+  /// True if the preview is currently undocked or in the process of undocking.
+  bool get _bIsUndockedOrUndocking => _dockState == _DockState.undocking || _dockState == _DockState.undocked;
+
+  /// True if the preview is not docked and offscreen.
+  bool get _bIsPreviewVisible {
+    if (_dockState != _DockState.docked) {
+      return true;
+    }
+
+    // Check if the preview is actually within visible bounds (e.g. still moving to docked position or being dragged out
+    // of it).
+    return ((_previewPosition.dx + _previewSize.width - _dockTabSize.width) > 0 &&
+        (_previewPosition.dx + _dockTabSize.width) < _containerSize.width);
+  }
+
+  @override
+  Widget build(BuildContext context) {
+    return LayoutBuilder(builder: (BuildContext context, BoxConstraints constraints) {
+      final Size newSize = Size(constraints.maxWidth, constraints.maxHeight);
+
+      // Handle first size/resize
+      if (newSize != _containerSize) {
+        final Size oldSize = _containerSize;
+        _containerSize = newSize;
+
+        if (oldSize == Size.zero) {
+          _initPosition();
+        } else {
+          WidgetsBinding.instance.addPostFrameCallback((_) => _onContainerResized(oldSize));
+        }
+      }
+
+      // Determine the offset of the main preview window, which shifts to the side to reveal dock tabs when docked.
+      double dockOffset = 0;
+      if (_bIsDockedOrDocking) {
+        switch (_dockSide) {
+          case _DockSide.left:
+            dockOffset = -_dockTabSize.width;
+            break;
+
+          case _DockSide.right:
+            dockOffset = _dockTabSize.width;
+            break;
+        }
+      }
+
+      return Stack(
+        clipBehavior: Clip.hardEdge,
+        children: [
+          // Position a box containing the preview, dock tabs, and gesture detector
+          Positioned(
+            top: _previewPosition.dy,
+            left: _previewPosition.dx,
+            child: SizedBox(
+              width: _previewSize.width,
+              height: _previewSize.height,
+              child: Stack(
+                clipBehavior: Clip.none,
+                children: [
+                  // Tab shown on the left when docked on the right
+                  if (_dockState != _DockState.undocked && _dockSide == _DockSide.right)
+                    Positioned(
+                      key: Key('LeftDockTab'),
+                      left: 0,
+                      top: (_previewSize.height - _dockTabSize.height) / 2,
+                      child: _DockTab(_DockSide.right),
+                    ),
+
+                  // Tab shown on the right when docked on the left
+                  if (_dockState != _DockState.undocked && _dockSide == _DockSide.left)
+                    Positioned(
+                      key: Key('RightDockTab'),
+                      left: _previewSize.width - _dockTabSize.width,
+                      top: (_previewSize.height - _dockTabSize.height) / 2,
+                      child: _DockTab(_DockSide.left),
+                    ),
+
+                  // Preview window, which shifts to reveal a tab when docked and hides when offscreen
+                  AnimatedPositioned(
+                    key: Key('PreviewMover'),
+                    left: dockOffset,
+                    duration: const Duration(milliseconds: 200),
+                    child: _bIsPreviewVisible
+                        ? _MapPreview(
+                            key: Key('Preview'),
+                            size: _previewSize,
+                            bShowShadow: _bIsUndockedOrUndocking,
+                          )
+                        : SizedBox(),
+                    onEnd: _onDockTabAnimationEnd,
+                  ),
+
+                  // Gesture detector, which matches the preview rectangle when undocked and the revealed tab when docked
+                  Positioned(
+                    key: Key('GestureDetector'),
+
+                    // When docked to left, offset to line up with the tab on the right side of the preview
+                    left: (_dockState == _DockState.docked && _dockSide == _DockSide.left)
+                        ? (_previewSize.width - _dockTabSize.width)
+                        : 0,
+
+                    // When docked, offset to line up with the top of the tab
+                    top: _dockState == _DockState.docked ? ((_previewSize.height - _dockTabSize.height) / 2) : 0,
+
+                    child: SizedBox(
+                      // When docked, match size to the tab
+                      width: _dockState == _DockState.docked ? _dockTabSize.width : _previewSize.width,
+                      height: _dockState == _DockState.docked ? _dockTabSize.height : _previewSize.height,
+
+                      child: GestureDetector(
+                        onTap: _onTap,
+                        onPanUpdate: _updateDrag,
+                        onPanEnd: (details) => _endDrag(details.velocity),
+                        onPanCancel: () => _endDrag(Velocity.zero),
+                      ),
+                    ),
+                  ),
+                ],
+              ),
+            ),
+          ),
+        ],
+      );
+    });
+  }
+
+  @override
+  void initState() {
+    super.initState();
+
+    _ticker = createTicker(_onTick);
+    _settings = FloatingMapSettings(PreferencesBundle.of(context));
+  }
+
+  @override
+  void dispose() {
+    _ticker.dispose();
+
+    super.dispose();
+  }
+
+  /// Initialize the position of the preview once we a valid size to reference.
+  /// This should NOT call setState, as it's called directly from the first build().
+  void _initPosition() {
+    final restoredYPosition = _settings.mapY.getValue() * _containerSize.height;
+    final PreviewMapSide mapSide = _settings.mapSide.getValue();
+
+    if (_settings.bIsMapDocked.getValue()) {
+      switch (mapSide) {
+        case PreviewMapSide.left:
+          _dockSide = _DockSide.left;
+          break;
+
+        case PreviewMapSide.right:
+          _dockSide = _DockSide.right;
+          break;
+      }
+
+      _dockState = _DockState.docked;
+      _targetPreviewPosition = _previewPosition = _getDockedPreviewPosition(dockY: restoredYPosition);
+    } else {
+      // Make an approximate position for the preview, which we'll then snap to the nearest corner
+      late final double tempXPosition;
+      switch (mapSide) {
+        case PreviewMapSide.left:
+          tempXPosition = 0;
+          break;
+
+        case PreviewMapSide.right:
+          tempXPosition = _containerSize.width - _previewSize.width;
+          break;
+      }
+
+      final tempPosition = Offset(tempXPosition, restoredYPosition);
+
+      _dockState = _DockState.undocked;
+      _targetPreviewPosition = _previewPosition = _getNearestPreviewCornerPosition(tempPosition);
+    }
+  }
+
+  /// Called when the containing widget is resized and we need to readjust accordingly.
+  void _onContainerResized(Size oldContainerSize) {
+    if (_containerSize.height == 0 || _containerSize.width == 0) {
+      return;
+    }
+
+    final Offset scaledPosition = Offset(
+      _targetPreviewPosition.dx * (_containerSize.width / oldContainerSize.width),
+      _targetPreviewPosition.dy * (_containerSize.height / oldContainerSize.height),
+    );
+
+    if (_bIsDockedOrDocking) {
+      // Stay docked and snap to the new side
+      setState(() {
+        _targetPreviewPosition = _previewPosition = _getDockedPreviewPosition(dockY: scaledPosition.dy);
+      });
+    } else {
+      // Move smoothly to the new position
+      _movePreviewTargetToNearestCorner(scaledPosition);
+    }
+  }
+
+  /// Handle a tap on the preview.
+  void _onTap() {
+    if (_bIsUndockedOrUndocking) {
+      return;
+    }
+
+    // Currently docked, so undock by moving to the nearest corner.
+    _movePreviewTargetToNearestCorner(_targetPreviewPosition);
+  }
+
+  /// Update the dragged preview position.
+  void _updateDrag(DragUpdateDetails details) {
+    _lastDragPosition = details.globalPosition;
+
+    _movePreviewTarget(
+      _targetPreviewPosition + Offset(details.delta.dx, details.delta.dy),
+      springDescription: _draggedSpringDescription,
+    );
+
+    final _DockSide? newDockSide = _checkDockSideAtPosition(_targetPreviewPosition);
+    if (newDockSide != null) {
+      // If this will dock, show the tab
+      _dockPreview(newDockSide, moveOffscreen: false);
+    } else {
+      // No longer docked, so move to undocking state
+      _undockPreview();
+    }
+  }
+
+  /// Stop dragging the preview and let it come to rest.
+  void _endDrag(Velocity velocity) {
+    final RenderBox? renderBox = context.findRenderObject() as RenderBox?;
+    if (renderBox == null) {
+      _movePreviewTarget(
+        Offset(_edgeInset, _edgeInset),
+        springDescription: _releasedSpringDescription,
+      );
+      return;
+    }
+
+    // Check ahead from the preview's position so the user can "fling" the preview off the screen
+    const dockCheckAheadSeconds = 0.1;
+    final Offset dockCheckPosition = _previewPosition + velocity.pixelsPerSecond * dockCheckAheadSeconds;
+
+    // If near the edges, dock the preview
+    final _DockSide? newDockSide = _checkDockSideAtPosition(dockCheckPosition);
+    if (newDockSide != null) {
+      _dockPreview(newDockSide, dockY: dockCheckPosition.dy);
+      return;
+    }
+
+    // Otherwise, move to the nearest corner. In this case, use the user's drag position instead of the preview
+    // position since flinging feels more accurate that way.
+    const cornerCheckAheadSeconds = 0.15;
+    final Offset cornerCheckPosition = _lastDragPosition + velocity.pixelsPerSecond * cornerCheckAheadSeconds;
+    final Offset finalLocalPosition = renderBox.globalToLocal(cornerCheckPosition);
+
+    _movePreviewTargetToNearestCorner(finalLocalPosition);
+  }
+
+  /// Update spring preview position using springs.
+  void _onTick(Duration elapsed) {
+    if (_xSpring == null || _ySpring == null) {
+      return;
+    }
+
+    _lastSpringTime = elapsed.inMilliseconds / 1000.0;
+
+    final double springTime = _lastSpringTime - _lastSpringMoveTime;
+
+    setState(() {
+      _previewVelocity = Offset(_xSpring!.dx(springTime), _ySpring!.dx(springTime));
+      _previewPosition = Offset(_xSpring!.x(springTime), _ySpring!.x(springTime));
+    });
+
+    // If both springs are at rest, stop ticking to save performance
+    if (_xSpring!.isDone(springTime) && _ySpring!.isDone(springTime)) {
+      _ticker.stop();
+      _lastSpringMoveTime = 0;
+      _lastSpringTime = 0;
+    }
+  }
+
+  /// Called when the animation to reveal/hide a dock tab completes.
+  void _onDockTabAnimationEnd() {
+    switch (_dockState) {
+      case _DockState.undocking:
+        setState(() {
+          _dockState = _DockState.undocked;
+        });
+        break;
+
+      case _DockState.docking:
+        setState(() {
+          _dockState = _DockState.docked;
+        });
+        break;
+
+      default:
+        break;
+    }
+  }
+
+  /// Move the target position for the preview to a new location and adjust springs accordingly.
+  void _movePreviewTarget(Offset newTarget, {required SpringDescription springDescription}) {
+    // Clamp both axes in a valid range
+    _targetPreviewPosition = Offset(
+      newTarget.dx.clamp(-_previewSize.width + _dockTabSize.width, _containerSize.width - _dockTabSize.width),
+      _getClampedPreviewYPosition(newTarget.dy),
+    );
+
+    _xSpring = SpringSimulation(
+      springDescription,
+      _previewPosition.dx,
+      _targetPreviewPosition.dx,
+      _previewVelocity.dx,
+    );
+
+    _ySpring = SpringSimulation(
+      springDescription,
+      _previewPosition.dy,
+      _targetPreviewPosition.dy,
+      _previewVelocity.dy,
+    );
+
+    _lastSpringMoveTime = _lastSpringTime;
+
+    if (!_ticker.isActive) {
+      _ticker.start();
+    }
+
+    _updateSettingsPreviewMapY(newTarget.dy);
+  }
+
+  /// Move the preview target to the corner nearest to a local [position].
+  void _movePreviewTargetToNearestCorner(Offset position) {
+    _movePreviewTarget(
+      _getNearestPreviewCornerPosition(position),
+      springDescription: _releasedSpringDescription,
+    );
+
+    _settings.mapSide.setValue(
+      (_targetPreviewPosition.dx < _containerSize.width / 2) ? PreviewMapSide.left : PreviewMapSide.right,
+    );
+
+    _undockPreview();
+  }
+
+  /// Reveal the dock tab and move to the docked state to the given [side].
+  /// If [moveOffscreen] is true, move the preview off-screen, optionally using the Y position specified by [dockY].
+  void _dockPreview(_DockSide side, {bool moveOffscreen = true, double? dockY}) {
+    dockY = dockY ?? _previewPosition.dy;
+
+    if (_bIsUndockedOrUndocking) {
+      // Trigger the docking animation
+      setState(() {
+        _dockSide = side;
+        _dockState = _DockState.docking;
+      });
+
+      // Update user settings
+      switch (side) {
+        case _DockSide.left:
+          _settings.mapSide.setValue(PreviewMapSide.left);
+          break;
+
+        case _DockSide.right:
+          _settings.mapSide.setValue(PreviewMapSide.right);
+          break;
+
+        default:
+          break;
+      }
+
+      _settings.bIsMapDocked.setValue(true);
+      _updateSettingsPreviewMapY(dockY);
+    }
+
+    if (!moveOffscreen) {
+      return;
+    }
+
+    _movePreviewTarget(
+      _getDockedPreviewPosition(dockY: dockY),
+      springDescription: _dockedSpringDescription,
+    );
+  }
+
+  /// Hide the docked tab and move to the undocking state.
+  void _undockPreview() {
+    if (_bIsDockedOrDocking) {
+      setState(() {
+        _dockState = _DockState.undocking;
+      });
+
+      _settings.bIsMapDocked.setValue(false);
+    }
+  }
+
+  /// Check whether the preview should be docked if it's at the given X position.
+  /// Returns the side to dock to, or null if it shouldn't dock.
+  _DockSide? _checkDockSideAtPosition(Offset position) {
+    if (position.dx + ((1 - _dockThreshold) * _previewSize.width) > _containerSize.width) {
+      return _DockSide.right;
+    }
+
+    if (position.dx + (_dockThreshold * _previewSize.width) < 0) {
+      return _DockSide.left;
+    }
+
+    return null;
+  }
+
+  /// Update the saved Y position of the docked preview.
+  void _updateSettingsPreviewMapY(double dockY) {
+    _settings.mapY.setValue(dockY / _containerSize.height);
+  }
+
+  /// Clamp a Y position for the preview to a valid range.
+  double _getClampedPreviewYPosition(double y) {
+    final double minY = _edgeInset;
+    final double maxY = _containerSize.height - _previewSize.height - _edgeInset;
+
+    return math.max(math.min(y, maxY), minY);
+  }
+
+  /// Get the position of the preview when moved to the nearest corner to a local [position].
+  Offset _getNearestPreviewCornerPosition(Offset position) {
+    late final double targetX;
+    if (position.dx < _containerSize.width / 2) {
+      targetX = _edgeInset;
+    } else {
+      targetX = _containerSize.width - _previewSize.width - _edgeInset;
+    }
+
+    late final double targetY;
+    if (position.dy < _containerSize.height / 2) {
+      targetY = _edgeInset;
+    } else {
+      targetY = _containerSize.height - _previewSize.height - _edgeInset;
+    }
+
+    return Offset(targetX, targetY);
+  }
+
+  /// Get the docked position of the preview when its Y position is [dockY].
+  Offset _getDockedPreviewPosition({required double dockY}) {
+    // Extra amount to move just to make sure it's fully off-screen so we can stop rendering the preview
+    const double offsetFudge = 0.5;
+
+    // Move the preview to the appropriate side
+    late final double targetX;
+    switch (_dockSide) {
+      case _DockSide.left:
+        targetX = -_previewSize.width + _dockTabSize.width - offsetFudge;
+        break;
+
+      case _DockSide.right:
+        targetX = _containerSize.width - _dockTabSize.width + offsetFudge;
+        break;
+    }
+
+    return Offset(targetX, _getClampedPreviewYPosition(dockY));
+  }
+}
+
+/// The static map preview which moves around within the larger widget.
+class _MapPreview extends StatelessWidget {
+  const _MapPreview({Key? key, required this.size, required this.bShowShadow}) : super(key: key);
+
+  final Size size;
+  final bool bShowShadow;
+
+  @override
+  Widget build(BuildContext context) {
+    final borderRadius = BorderRadius.circular(6);
+
+    return SizedBox(
+      width: size.width,
+      height: size.height,
+      child: Stack(
+        children: [
+          AnimatedContainer(
+            clipBehavior: Clip.antiAlias,
+            duration: Duration(milliseconds: 400),
+            decoration: BoxDecoration(
+              borderRadius: borderRadius,
+              boxShadow: [
+                BoxShadow(
+                  color: Color(0xff000000).withOpacity(bShowShadow ? 0.4 : 0),
+                  spreadRadius: 2,
+                  blurRadius: 5,
+                ),
+              ],
+            ),
+            child: StageMap(
+              bIsControllable: false,
+              bShowOnlySelectedPins: true,
+              bForceShowFullMap: true,
+              pinSize: 32,
+            ),
+          ),
+          Container(
+            decoration: BoxDecoration(
+              borderRadius: borderRadius,
+              border: Border.all(
+                color: Color(0x50ffffff),
+                width: 0.5,
+              ),
+            ),
+          ),
+        ],
+      ),
+    );
+  }
+}
+
+/// The tab shown when the preview is docked.
+class _DockTab extends StatelessWidget {
+  const _DockTab(this.side, {Key? key}) : super(key: key);
+
+  /// Which side the preview is docked to.
+  final _DockSide side;
+
+  @override
+  Widget build(BuildContext context) {
+    const Radius borderRadius = Radius.circular(8);
+    final bool bIsFacingLeft = side == _DockSide.right;
+
+    return Container(
+      width: _dockTabSize.width,
+      height: _dockTabSize.height,
+      decoration: BoxDecoration(
+        color: Theme.of(context).colorScheme.surfaceTint,
+        borderRadius: bIsFacingLeft
+            ? BorderRadius.only(
+                topLeft: borderRadius,
+                bottomLeft: borderRadius,
+              )
+            : BorderRadius.only(
+                topRight: borderRadius,
+                bottomRight: borderRadius,
+              ),
+        boxShadow: [
+          BoxShadow(
+            color: Color(0x40000000),
+            offset: Offset(bIsFacingLeft ? -4 : 4, 4),
+            blurRadius: 4,
+          ),
+        ],
+      ),
+      child: Column(mainAxisAlignment: MainAxisAlignment.center, children: [
+        AssetIcon(
+          path: 'assets/images/icons/map.svg',
+          size: 24,
+        ),
+        SizedBox(
+          height: 6,
+        ),
+        AssetIcon(
+          path: bIsFacingLeft ? 'assets/images/icons/chevron_left.svg' : 'assets/images/icons/chevron_right.svg',
+          size: 24,
+        ),
+      ]),
+    );
+  }
+}
diff --git a/Engine/Extras/VirtualProduction/EpicStageApp/lib/widgets/elements/layout/card.dart b/Engine/Extras/VirtualProduction/EpicStageApp/lib/widgets/elements/layout/card.dart
new file mode 100644
--- /dev/null
+++ b/Engine/Extras/VirtualProduction/EpicStageApp/lib/widgets/elements/layout/card.dart
@@ -0,0 +1,339 @@
+// Copyright Epic Games, Inc. All Rights Reserved.
+
+import 'package:flutter/material.dart';
+
+import '../../../utilities/unreal_colors.dart';
+import '../asset_icon.dart';
+import '../epic_icon_button.dart';
+
+/// A large title bar for a [Card] including an icon and subtitle.
+class CardLargeHeader extends StatelessWidget {
+  const CardLargeHeader({
+    Key? key,
+    this.iconPath,
+    this.title,
+    this.subtitle,
+    this.trailing,
+  }) : super(key: key);
+
+  final String? iconPath;
+  final String? title;
+  final String? subtitle;
+
+  /// Optional widget to show on the opposite side to the header.
+  final Widget? trailing;
+
+  @override
+  Widget build(BuildContext context) {
+    return Container(
+      height: 44,
+      padding: const EdgeInsets.only(left: 16),
+      alignment: Alignment.centerLeft,
+      color: Theme.of(context).colorScheme.surfaceTint,
+      child: Row(
+        children: [
+          if (iconPath != null)
+            Padding(
+              padding: EdgeInsets.only(right: 16),
+              child: AssetIcon(path: iconPath!, size: 24),
+            ),
+          Expanded(
+            child: RichText(
+              overflow: TextOverflow.ellipsis,
+              text: TextSpan(children: [
+                if (title != null) TextSpan(text: title, style: Theme.of(context).textTheme.displayLarge),
+                if (subtitle != null)
+                  TextSpan(
+                    text: (title != null) ? ' – $subtitle' : subtitle,
+                    style: Theme.of(context).textTheme.displayMedium,
+                  ),
+              ]),
+            ),
+          ),
+          if (trailing != null) trailing!,
+        ],
+      ),
+    );
+  }
+}
+
+/// A small title bar for a [Card].
+class CardSmallHeader extends StatelessWidget {
+  const CardSmallHeader({
+    Key? key,
+    required this.title,
+  }) : super(key: key);
+
+  final String title;
+
+  @override
+  Widget build(BuildContext context) {
+    return Container(
+      height: 44,
+      padding: const EdgeInsets.symmetric(horizontal: 16, vertical: 10),
+      alignment: Alignment.centerLeft,
+      color: Theme.of(context).colorScheme.surfaceTint,
+      child: Text(
+        title.toUpperCase(),
+        style: Theme.of(context).textTheme.titleLarge,
+      ),
+    );
+  }
+}
+
+/// Possible states for displaying a [CardListTile] which may have foldout children.
+enum CardListTileExpansionState {
+  /// This can't be expanded or collapsed.
+  none,
+
+  /// The list of children is currently collapsed.
+  collapsed,
+
+  /// The list of children is currently expanded.
+  expanded,
+}
+
+/// A list tile shown in a [Card].
+class CardListTile extends StatelessWidget {
+  const CardListTile({
+    Key? key,
+    required this.title,
+    required this.bIsSelected,
+    this.stateIconPath,
+    this.iconPath,
+    this.trailing,
+    this.onTap,
+    this.onTapIcon,
+    this.bDeEmphasize = false,
+    this.expansionState = CardListTileExpansionState.none,
+    this.indentation = 0,
+  }) : super(key: key);
+
+  /// The label to display on the item.
+  final String title;
+
+  /// Whether this is currently selected.
+  final bool bIsSelected;
+
+  /// If provided, show a small icon with the image at this path before the rest of the contents to indicate state.
+  final String? stateIconPath;
+
+  /// If provided, show an icon with the image at this path.
+  final String? iconPath;
+
+  /// Optional widget to show after the title.
+  final Widget? trailing;
+
+  /// Function called when the user taps this item.
+  final void Function()? onTap;
+
+  /// Function called when the user taps the icon, if present.
+  final void Function()? onTapIcon;
+
+  /// If true, fade the text and icon of this panel.
+  final bool bDeEmphasize;
+
+  /// Current state indicating whether this is expandable and if so, whether it's currently expanded.
+  final CardListTileExpansionState expansionState;
+
+  /// How many layers deep this item is indented.
+  final int indentation;
+
+  @override
+  Widget build(BuildContext context) {
+    final ThemeData theme = Theme.of(context);
+    final ColorScheme colorScheme = theme.colorScheme;
+
+    final double opacity = bDeEmphasize ? 0.4 : 1.0;
+    final Color textColor = (bIsSelected ? colorScheme.onPrimary : colorScheme.onSurface).withOpacity(opacity);
+    final TextStyle textStyle = theme.textTheme.headlineSmall!.copyWith(
+      color: textColor,
+      fontStyle: bDeEmphasize ? FontStyle.italic : FontStyle.normal,
+    );
+
+    Widget? stateIcon = null;
+    if (stateIconPath != null) {
+      stateIcon = AssetIcon(
+        path: stateIconPath!,
+        size: 16,
+      );
+    }
+
+    Widget? icon = null;
+    if (iconPath != null) {
+      icon = AssetIcon(
+        path: iconPath!,
+        size: 24,
+        color: Color.fromRGBO(255, 255, 255, opacity),
+      );
+    }
+
+    final bool bIsExpandable = (expansionState == CardListTileExpansionState.collapsed) ||
+        (expansionState == CardListTileExpansionState.expanded);
+
+    Color color = Colors.transparent;
+    if (bIsSelected) {
+      color = Theme.of(context).colorScheme.primary;
+    } else if (bIsExpandable) {
+      color = UnrealColors.gray13;
+    }
+
+    final Widget mainBody = Container(
+      color: color,
+      height: 40,
+      child: ListTile(
+        leading: Row(
+          mainAxisSize: MainAxisSize.min,
+          children: [
+            if (indentation > 0) SizedBox(width: 20.0 * indentation),
+            if (stateIcon != null)
+              Padding(
+                padding: EdgeInsets.only(right: 10),
+                child: stateIcon,
+              ),
+            SizedBox.square(dimension: 24, child: icon),
+          ],
+        ),
+        trailing: trailing,
+        title: Text(
+          title,
+          style: textStyle,
+          softWrap: false,
+          overflow: TextOverflow.ellipsis,
+        ),
+        horizontalTitleGap: 8,
+        minLeadingWidth: 0,
+        contentPadding: EdgeInsets.only(
+          left: 14,
+          right: 8,
+        ),
+        visualDensity: const VisualDensity(vertical: -4),
+        onTap: onTap,
+      ),
+    );
+
+    if (expansionState == CardListTileExpansionState.none) {
+      return mainBody;
+    }
+
+    // Place a gesture detector to intercept taps on the icon's area
+    if (icon != null && onTapIcon != null) {
+      return Stack(
+        clipBehavior: Clip.none,
+        children: [
+          mainBody,
+          Positioned(
+            top: 0,
+            bottom: 0,
+            width: 48,
+            child: GestureDetector(
+              onTap: onTapIcon,
+            ),
+          ),
+        ],
+      );
+    } else {
+      return mainBody;
+    }
+  }
+}
+
+/// A button revealed in a [SwipeRevealer] in conjunction with a [CardListTile].
+class CardListTileSwipeAction extends StatelessWidget {
+  const CardListTileSwipeAction({
+    Key? key,
+    required this.iconPath,
+    required this.color,
+    required this.onPressed,
+    this.iconSize = 24,
+  }) : super(key: key);
+
+  /// The path of the image file to show inside the button.
+  final String iconPath;
+
+  /// Color of the button background.
+  final Color color;
+
+  /// Function to call when this is pressed.
+  final Function() onPressed;
+
+  /// The size of the icon (both width and height).
+  final double iconSize;
+
+  @override
+  Widget build(BuildContext context) {
+    return EpicIconButton(
+      iconPath: iconPath,
+      bIsToggledOn: true,
+      iconSize: iconSize,
+      activeBackgroundColor: color,
+      buttonSize: const Size(64, 40),
+      onPressed: onPressed,
+    );
+  }
+}
+
+/// A sub-header in a [Card].
+class CardSubHeader extends StatelessWidget {
+  const CardSubHeader({
+    Key? key,
+    required this.child,
+    this.padding,
+    this.height = 48,
+  }) : super(key: key);
+
+  final Widget child;
+  final EdgeInsetsGeometry? padding;
+  final double height;
+
+  @override
+  Widget build(BuildContext context) {
+    return Container(
+      padding: padding,
+      height: height,
+      color: Theme.of(context).colorScheme.surfaceTint,
+      child: child,
+    );
+  }
+}
+
+/// An icon button shown in a card sub-header.
+class CardSubHeaderButton extends StatelessWidget {
+  const CardSubHeaderButton({
+    Key? key,
+    required this.iconPath,
+    required this.tooltipMessage,
+    this.onPressed,
+    this.bIsToggledOn = false,
+    this.bIsVisualOnly = false,
+  }) : super(key: key);
+
+  /// The path of the image file to show inside the button.
+  final String iconPath;
+
+  /// An optional tooltip message to show when the button is hovered or long pressed.
+  final String tooltipMessage;
+
+  /// The function to call when the button is pressed. If not provided and [bIsVisualOnly] is false, the button is
+  /// considered inactive and will change color accordingly.
+  final Function()? onPressed;
+
+  /// If true, highlight the button using the [activeBackgroundColor].
+  final bool bIsToggledOn;
+
+  /// If true, leave out button semantics and only include the button's visuals.
+  /// This will disable the tooltip message.
+  final bool bIsVisualOnly;
+
+  @override
+  Widget build(BuildContext context) {
+    return EpicIconButton(
+      iconPath: iconPath,
+      tooltipMessage: tooltipMessage,
+      onPressed: onPressed,
+      buttonSize: const Size(48, 48),
+      bIsToggledOn: bIsToggledOn,
+      bIsVisualOnly: bIsVisualOnly,
+    );
+  }
+}
diff --git a/Engine/Extras/VirtualProduction/EpicStageApp/lib/widgets/elements/layout/lazy_indexed_stack.dart b/Engine/Extras/VirtualProduction/EpicStageApp/lib/widgets/elements/layout/lazy_indexed_stack.dart
new file mode 100644
--- /dev/null
+++ b/Engine/Extras/VirtualProduction/EpicStageApp/lib/widgets/elements/layout/lazy_indexed_stack.dart
@@ -0,0 +1,110 @@
+// Copyright Epic Games, Inc. All Rights Reserved.
+
+import 'dart:math';
+
+import 'package:flutter/material.dart';
+
+/// An IndexedStack that doesn't create its child widgets until they need to be displayed.
+class LazyIndexedStack extends StatefulWidget {
+  const LazyIndexedStack({
+    Key? key,
+    this.alignment = AlignmentDirectional.topStart,
+    this.textDirection,
+    this.sizing = StackFit.loose,
+    this.clipBehavior = Clip.hardEdge,
+    this.index = 0,
+    this.keepAliveChildren,
+    required this.childBuilders,
+  }) : super(key: key);
+
+  /// How to align the non-positioned and partially-positioned children in the
+  /// stack.
+  final AlignmentGeometry alignment;
+
+  /// The text direction with which to resolve [alignment].
+  final TextDirection? textDirection;
+
+  /// How to size the non-positioned children in the stack.
+  final StackFit sizing;
+
+  /// How to handle clipping for contents of the stack.
+  final Clip clipBehavior;
+
+  /// The index into [childBuilders] for the widget to display.
+  final int index;
+
+  /// Functions used to build each child.
+  /// They won't be built until their index has been set as [childIndex] at least once.
+  final List<Function(BuildContext context)> childBuilders;
+
+  /// An optional list corresponding to [childBuilders], indicating which children should be kept alive when their index
+  /// is no longer selected. If not provided, all children will be disposed when their index is not selected.
+  final List<bool>? keepAliveChildren;
+
+  @override
+  State<LazyIndexedStack> createState() => _LazyIndexedStackState();
+}
+
+class _LazyIndexedStackState extends State<LazyIndexedStack> {
+  static const Widget _placeholderWidget = SizedBox();
+
+  List<Widget> _displayedChildren = [];
+
+  @override
+  void initState() {
+    super.initState();
+
+    _displayedChildren = widget.childBuilders.map((_) => _placeholderWidget).toList();
+    _buildChild(widget.index);
+  }
+
+  @override
+  void didUpdateWidget(covariant LazyIndexedStack oldWidget) {
+    super.didUpdateWidget(oldWidget);
+
+    // Shrink to match new children count
+    while (_displayedChildren.length > widget.childBuilders.length) {
+      _displayedChildren.removeAt(_displayedChildren.length);
+    }
+
+    // Grow to match new children count
+    while (_displayedChildren.length < widget.childBuilders.length) {
+      _displayedChildren.add(_placeholderWidget);
+    }
+
+    // Replace any changed children with placeholders until they're loaded
+    final maxChangedChildIndex = min(widget.childBuilders.length, oldWidget.childBuilders.length);
+    for (int childIndex = 0; childIndex < maxChangedChildIndex; ++childIndex) {
+      if (widget.childBuilders[childIndex] != oldWidget.childBuilders[childIndex]) {
+        _displayedChildren[childIndex] = _placeholderWidget;
+      }
+    }
+
+    if (oldWidget.index != widget.index) {
+      // Build and remember the current child
+      _buildChild(widget.index);
+
+      // Forget the old child
+      if (widget.keepAliveChildren == null ||
+          oldWidget.index >= widget.keepAliveChildren!.length ||
+          !widget.keepAliveChildren![oldWidget.index]) {
+        _displayedChildren[oldWidget.index] = _placeholderWidget;
+      }
+    }
+  }
+
+  @override
+  Widget build(BuildContext context) {
+    return IndexedStack(
+      alignment: widget.alignment,
+      textDirection: widget.textDirection,
+      sizing: widget.sizing,
+      index: widget.index,
+      children: _displayedChildren,
+    );
+  }
+
+  void _buildChild(int childIndex) {
+    _displayedChildren[widget.index] = widget.childBuilders[widget.index](context);
+  }
+}
diff --git a/Engine/Extras/VirtualProduction/EpicStageApp/lib/widgets/elements/layout/lazy_tab_view.dart b/Engine/Extras/VirtualProduction/EpicStageApp/lib/widgets/elements/layout/lazy_tab_view.dart
new file mode 100644
--- /dev/null
+++ b/Engine/Extras/VirtualProduction/EpicStageApp/lib/widgets/elements/layout/lazy_tab_view.dart
@@ -0,0 +1,87 @@
+// Copyright Epic Games, Inc. All Rights Reserved.
+
+import 'package:flutter/material.dart';
+
+import 'lazy_indexed_stack.dart';
+
+/// A view that lazily creates the contents of each tab as needed and can optionally discard those contents when the
+/// tab is no longer selected.
+class LazyTabView extends StatefulWidget {
+  const LazyTabView({Key? key, required this.controller, required this.builder, this.keepAlive}) : super(key: key);
+
+  /// The tab controller that controls which tab to show.
+  final TabController controller;
+
+  /// A builder function that creates the tab contents.
+  final Function(BuildContext context, int tabIndex) builder;
+
+  /// An optional function that determines whether a tab should be kept alive when its contents are invisible.
+  final bool Function(int tabIndex)? keepAlive;
+
+  @override
+  State<LazyTabView> createState() => _LazyTabViewState();
+}
+
+class _LazyTabViewState extends State<LazyTabView> {
+  /// The index of the currently displayed tab.
+  int _tabIndex = -1;
+  List<Function(BuildContext context)> _childBuilders = [];
+  List<bool> _keepAliveChildren = [];
+
+  @override
+  void initState() {
+    super.initState();
+
+    widget.controller.addListener(_onTabControllerChanged);
+    _tabIndex = widget.controller.index;
+
+    _updateChildBuilders();
+  }
+
+  @override
+  void dispose() {
+    widget.controller.removeListener(_onTabControllerChanged);
+
+    super.dispose();
+  }
+
+  @override
+  void didUpdateWidget(covariant LazyTabView oldWidget) {
+    super.didUpdateWidget(oldWidget);
+
+    if (oldWidget.builder != widget.builder) {
+      _updateChildBuilders();
+    }
+  }
+
+  @override
+  Widget build(BuildContext context) {
+    return LazyIndexedStack(
+      childBuilders: _childBuilders,
+      index: _tabIndex,
+      keepAliveChildren: _keepAliveChildren,
+    );
+  }
+
+  /// Update the list of functions used to build children of the tab view.
+  void _updateChildBuilders() {
+    _childBuilders = [];
+    _keepAliveChildren = [];
+
+    for (int childIndex = 0; childIndex < widget.controller.length; ++childIndex) {
+      _childBuilders.add((BuildContext context) => widget.builder(context, childIndex));
+      _keepAliveChildren.add(widget.keepAlive != null ? widget.keepAlive!(childIndex) : false);
+    }
+  }
+
+  /// Called when the tab controller updates and may have changed tabs.
+  void _onTabControllerChanged() {
+    if (widget.controller.index == _tabIndex) {
+      return;
+    }
+
+    setState(() {
+      _tabIndex = widget.controller.index;
+    });
+  }
+}
diff --git a/Engine/Extras/VirtualProduction/EpicStageApp/lib/widgets/elements/layout/measurable.dart b/Engine/Extras/VirtualProduction/EpicStageApp/lib/widgets/elements/layout/measurable.dart
new file mode 100644
--- /dev/null
+++ b/Engine/Extras/VirtualProduction/EpicStageApp/lib/widgets/elements/layout/measurable.dart
@@ -0,0 +1,43 @@
+// Copyright Epic Games, Inc. All Rights Reserved.
+
+import 'package:flutter/material.dart';
+import 'package:flutter/rendering.dart';
+
+/// A widget that reports the size of its child as soon as it's laid out.
+class Measurable extends SingleChildRenderObjectWidget {
+  const Measurable({Key? key, required super.child, required this.onMeasured}) : super(key: key);
+
+  /// Callback function which will be called whenever the [child] is laid out and reports its new [size].
+  final void Function(Size size) onMeasured;
+
+  @override
+  RenderObject createRenderObject(BuildContext context) => _MeasurableRenderObject(onMeasured);
+}
+
+class _MeasurableRenderObject extends RenderProxyBox {
+  _MeasurableRenderObject(this.onMeasured);
+
+  /// Callback function which will be called whenever the [child] is laid out and reports its new [size].
+  final void Function(Size size) onMeasured;
+
+  /// The previous size of this object.
+  Size? _lastSize;
+
+  @override
+  void performLayout() {
+    super.performLayout();
+
+    final Size? size = child?.size;
+    if (size == _lastSize) {
+      return;
+    }
+
+    _lastSize = size;
+    if (size == null) {
+      return;
+    }
+
+    // Report the new size post-frame to avoid triggering a rebuild in the middle of another rebuild.
+    WidgetsBinding.instance.addPostFrameCallback((_) => onMeasured(size));
+  }
+}
