Branch snapshot for CL 27138787

diff --git a/Engine/Extras/VirtualProduction/EpicStageApp/lib/widgets/elements/lightcard_map.dart b/Engine/Extras/VirtualProduction/EpicStageApp/lib/widgets/elements/lightcard_map.dart
new file mode 100644
--- /dev/null
+++ b/Engine/Extras/VirtualProduction/EpicStageApp/lib/widgets/elements/lightcard_map.dart
@@ -0,0 +1,1966 @@
+// Copyright Epic Games, Inc. All Rights Reserved.
+
+import 'dart:async';
+import 'dart:math' as math;
+import 'dart:typed_data';
+
+import 'package:flutter/gestures.dart';
+import 'package:flutter/material.dart';
+import 'package:flutter_gen/gen_l10n/app_localizations.dart';
+import 'package:logging/logging.dart';
+import 'package:provider/provider.dart';
+import 'package:streaming_shared_preferences/streaming_shared_preferences.dart';
+import 'package:vector_math/vector_math_64.dart' as vec;
+
+import '../../../models/engine_connection.dart';
+import '../../models/preview_render_manager.dart';
+import '../../models/property_modify_operations.dart';
+import '../../models/settings/selected_actor_settings.dart';
+import '../../models/settings/stage_map_settings.dart';
+import '../../models/unreal_actor_manager.dart';
+import '../../models/unreal_class_data.dart';
+import '../../models/unreal_property_manager.dart';
+import '../../models/unreal_types.dart';
+import '../../utilities/constants.dart';
+import '../../utilities/guarded_refresh_state.dart';
+import '../../utilities/transient_preference.dart';
+import '../../utilities/unreal_colors.dart';
+import 'asset_icon.dart';
+import 'dropdown_button.dart';
+import 'dropdown_list_menu.dart';
+import 'epic_icon_button.dart';
+import 'place_actor_menu.dart';
+import 'spinner_overlay.dart';
+import 'transform_gesture_detector.dart';
+
+final _log = Logger('LightcardMap');
+
+enum _GestureState {
+  none,
+  // Pinching to scale and/or rotate in transform mode.
+  // Flutter calls this a "scale" gesture so we use the same name for consistency.
+  transformScaling,
+  // Moving a single finger around the screen. This includes a pan gesture in transform mode or a press-and-hold gesture
+  // on a light card in camera mode.
+  transformDragging,
+  // Dragging with a single finger to rotate the camera.
+  cameraTumbling,
+  // Dragging with two fingers to move the camera or zoom it in.
+  cameraTrackingAndZooming,
+}
+
+/// A map displaying all stage actors pins on an azimuthal projection map of the scene.
+class StageMap extends StatefulWidget {
+  const StageMap({
+    Key? key,
+    this.bIsControllable = true,
+    this.bShowOnlySelectedPins = false,
+    this.bForceShowFullMap = false,
+    this.pinSize,
+  }) : super(key: key);
+
+  /// If false, disable map controls.
+  final bool bIsControllable;
+
+  /// If true, only show pins for actors that the user has selected.
+  final bool bShowOnlySelectedPins;
+
+  /// If true, ignore the user's zoom/pan settings and show the full map.
+  final bool bForceShowFullMap;
+
+  /// If provided, overrides the size of actor pins in pixels.
+  final double? pinSize;
+
+  @override
+  State<StageMap> createState() => StageMapState();
+}
+
+class StageMapState extends State<StageMap> with PreviewRenderConsumer, GuardedRefreshState {
+  /// The minimum scale of the preview map.
+  static const double _minScale = 1.0;
+
+  /// The maximum scale of the preview map.
+  static const double _maxScale = 2.5;
+
+  /// How much to scale based on scroll wheel input. This is based on InteractiveViewer's implementation, which uses the
+  /// same default value.
+  static const double _scrollScaleFactor = 200.0;
+
+  /// The approximate number of pixels to include in the preview's resolution. The actual resolution will depend on the
+  /// widget's aspect ratio.
+  static const int _previewResolutionPixels = 1024 * 1024;
+
+  /// Rate at which we send light card/camera config updates back to the engine.
+  static const Duration _engineUpdateRate = Duration(milliseconds: 100);
+
+  /// Rate at which the camera moves in each direction when tumbling with a drag gesture (in degrees per pixel dragged).
+  static const Offset _tumbleSensitivity = Offset(0.3, 0.3);
+
+  /// Distance from the initial touch point at which we'll actually start a tumble gesture using a single touch.
+  /// This makes it easier to scale/zoom with two fingers without accidentally starting a tumble gesture first.
+  static const double _tumbleTouchSlop = 15.0;
+
+  /// The current preview image.
+  Image? _previewImage;
+
+  /// The stream for the current preview image, used to listen for when it's ready
+  ImageStream? _previewImageStream;
+
+  /// Timer used to regularly send changes back to the engine while connected.
+  Timer? _engineUpdateTimer;
+
+  /// The new drag position to send on the next tick, or null if there's nothing new to send.
+  Offset? _tickDragPosition;
+
+  /// The last position of the current gesture's focal point in the draggable map scene's coordinate space.
+  Offset _gestureReferenceScenePoint = Offset.zero;
+
+  /// The global position at which the current pan/drag operation started.
+  Offset _gestureStartPosition = Offset.zero;
+
+  /// The map's scale at the start of the current scale gesture.
+  double _gestureStartScale = 1.0;
+
+  /// The offset between the pointer position and the dragged actor when the drag operation started.
+  Offset _dragPointerOffset = Offset.zero;
+
+  /// The last received scale update details during the current scale gesture.
+  ScaleUpdateDetails _lastScaleUpdateDetails = ScaleUpdateDetails();
+
+  /// The path of the pin currently being dragged.
+  String? _draggedPinPath;
+
+  /// Whether this is waiting for light card properties to be exposed before they can be transformed.
+  bool _bIsWaitingForExpose = false;
+
+  /// Whether to show a warning that the engine's real-time nDisplay updates are disabled.
+  bool _bShowRealTimeDisabledWarning = false;
+
+  /// The current gesture being performed.
+  _GestureState _gestureState = _GestureState.none;
+
+  /// The size at which to request preview renders.
+  Size _desiredPreviewSize = Size.zero;
+
+  /// The size of the last preview received, which may still be in the process of loading.
+  Size _pendingPreviewSize = Size.zero;
+
+  /// The size of the currently displayed preview image.
+  Size? _previewSize;
+
+  /// The last size of the render box in on-screen pixels.
+  Size _constraintSize = Size.zero;
+
+  /// Actor positions which we've received from the engine, but are waiting for the latest preview image to load before
+  /// displaying them to the user. This prevents the positions from updating before the image and appearing out of sync.
+  /// If null, there's no pending change.
+  List<dynamic> _pendingActorPositions = [];
+
+  /// Stream subscriptions to user settings.
+  final List<StreamSubscription> _settingSubscriptions = [];
+
+  /// Key used to refer to the map widget.
+  final GlobalKey _mapKey = GlobalKey();
+
+  /// Controller used to transform the map and its contents.
+  final TransformationController _mapTransform = TransformationController();
+
+  /// Map from actor path to the data we have stored about it.
+  final Map<String, _MapActorData> _mapActorDataByPath = {};
+
+  /// Paths of actors that have been directly placed by users.
+  /// Their new positions in _pendingActorPositions won't be applied when the image loads to prevent jittering.
+  /// This list is cleared whenever we receive a render that occurred after all drag inputs so far.
+  final Set<String> _recentDirectDragPins = {};
+
+  /// Connection manager to communicate with the engine.
+  late final EngineConnectionManager _connectionManager;
+
+  /// Actor manager used to look up information about actors visible on the map.
+  late final UnrealActorManager _actorManager;
+
+  /// The property manager used to control actor properties other than position (which is handled via a separate
+  /// WebSocket API).
+  late final UnrealPropertyManager _propertyManager;
+
+  /// Listener for image stream events, used to detect when the preview image loads.
+  late final PreviewRenderManager _previewRenderManager;
+
+  /// Listener for image stream events, used to detect when the preview image loads.
+  late final ImageStreamListener _previewImageStreamListener;
+
+  /// The rectangle containing the map.
+  Rect get _containerRect {
+    return Offset.zero & _constraintSize;
+  }
+
+  /// The aspect ratio of the container, or 0 if invalid.
+  double get _containerAspectRatio {
+    return _constraintSize.height > 0 ? _constraintSize.width / _constraintSize.height : 0;
+  }
+
+  @override
+  Widget build(BuildContext context) {
+    return Container(
+      color: UnrealColors.black,
+      child: LayoutBuilder(
+        builder: (BuildContext context, BoxConstraints constraints) {
+          _updateConstraints(constraints);
+
+          final stageMapSettings = Provider.of<StageMapSettings>(context, listen: false);
+          return PreferenceBuilder(
+            preference: stageMapSettings.bIsInTransformMode,
+            builder: (BuildContext context, bool bIsInTransformMode) {
+              final selectedActorSettings = Provider.of<SelectedActorSettings>(context, listen: false);
+
+              bool bShouldShowSpinner = false;
+              if (_previewImage == null) {
+                // We don't have a map image yet
+                if (_connectionManager.connectionState == EngineConnectionState.connected &&
+                    selectedActorSettings.displayClusterRootPath.getValue().length > 0) {
+                  // We're connected but still waiting for a map
+                  bShouldShowSpinner = true;
+                } else {
+                  // We're disconnected and no map is pending, so a spinner would be misleading
+                  return Container(color: Colors.black);
+                }
+              }
+
+              if (bShouldShowSpinner) {
+                return const SpinnerOverlay();
+              }
+
+              final Matrix4 invMapTransform = _mapTransform.value.clone();
+              invMapTransform.invert();
+
+              // Create the pins. Draggable (used in the pin) doesn't play nicely when its containing stack is
+              // reordered, so don't try to reorder or remove items from this list in the middle of a drag operation.
+              final Map<_StageMapPin, Offset> pinLocations = {};
+              for (MapEntry<String, _MapActorData> actorPair in _mapActorDataByPath.entries) {
+                // Hide deselected pins if requested
+                if (widget.bShowOnlySelectedPins && !selectedActorSettings.isActorSelected(actorPair.key)) {
+                  continue;
+                }
+
+                final UnrealObject? actor = _actorManager.getActorAtPath(actorPair.key);
+                if (actor == null) {
+                  continue;
+                }
+
+                // Hide actors belonging to other DCRAs
+                if (actor.unrealClass?.isA(UnrealClassRegistry.lightCard) == true &&
+                    !_actorManager.currentRootActorLightCards.value.contains(actor)) {
+                  continue;
+                }
+
+                final Offset actorPosition = actorPair.value.position;
+                final Offset pinOffset = _getPinOffset(actorPosition);
+
+                final _StageMapPin pin = widget.bIsControllable
+                    ? _StageMapPin(
+                        actor: actor,
+                        bIsDraggable: !bIsInTransformMode,
+                        bShowAsDragged: _shouldForcePinDragAnimation(actorPair.key),
+                        onDragStarted: () => _startPinDrag(actorPair.key),
+                        onDragUpdate: (Offset pinPosition) => _updatePinDrag(pinPosition),
+                        onDragEnded: (Offset dropPosition) => _endPinDrag(dropPosition),
+                      )
+                    : _StageMapPin(
+                        actor: actor,
+                        bIsDraggable: false,
+                        bIsSelectable: false,
+                        size: widget.pinSize,
+                      );
+
+                pinLocations[pin] = pinOffset;
+              }
+
+              // The contents of the map that can be scrolled around within the interactive viewer
+              final Widget mapContent = Transform(
+                transform: _mapTransform.value,
+                child: Stack(
+                  clipBehavior: Clip.none,
+                  alignment: Alignment.topLeft,
+                  children: [
+                    Positioned.fill(
+                      child: _previewImage!,
+                    ),
+
+                    // Actor pins
+                    Transform(
+                      // Apply the inverse transform so the pins stay at the same scale regardless of zoom.
+                      // _getPinOffset will take the map's offset into account while also adjusting for scale.
+                      transform: invMapTransform,
+                      child: Stack(clipBehavior: Clip.none, children: [
+                        for (MapEntry<_StageMapPin, Offset> pinPair in pinLocations.entries)
+                          Positioned(
+                            left: pinPair.value.dx,
+                            top: pinPair.value.dy,
+                            child: pinPair.key,
+                          )
+                      ]),
+                    ),
+                  ],
+                ),
+              );
+
+              final Widget map = widget.bIsControllable
+                  ? Listener(
+                      onPointerDown: _onPointerDown,
+                      onPointerMove: _onPointerMove,
+                      onPointerUp: _onPointerUp,
+                      onPointerCancel: _onPointerCancel,
+                      onPointerSignal: _onPointerSignal,
+                      child: TransformGestureDetector(
+                        onScaleStart: _onScaleStart,
+                        onScaleUpdate: _onScaleUpdate,
+                        onScaleEnd: _onScaleEnd,
+                        onLongPressStart: bIsInTransformMode ? null : _onLongPressStart,
+                        child: mapContent,
+                      ),
+                    )
+                  : mapContent;
+
+              return ClipRect(
+                clipBehavior: Clip.hardEdge,
+                child: Stack(
+                  children: [
+                    map,
+
+                    // Spinner shown when we have map data but the user can't change some properties yet
+                    if (_bIsWaitingForExpose) const SpinnerOverlay(),
+
+                    if (widget.bIsControllable)
+                      // Mode controls
+                      Positioned(
+                        top: 0,
+                        child: Padding(
+                          padding: const EdgeInsets.all(12),
+                          child: Container(
+                            padding: const EdgeInsets.all(6),
+                            decoration: BoxDecoration(
+                              borderRadius: BorderRadius.circular(40),
+                              color: const Color(0x4da1a1a1),
+                            ),
+                            child: Wrap(
+                              spacing: 16,
+                              children: [
+                                const _StageMapControlModeToggle(),
+                                ConstrainedBox(
+                                  constraints: const BoxConstraints(minWidth: 128),
+                                  child: Tooltip(
+                                    message: AppLocalizations.of(context)!.stageMapProjectionMode,
+                                    child: PreferenceBuilder<ProjectionMode>(
+                                      preference: stageMapSettings.projectionMode,
+                                      builder: (context, projectionMode) {
+                                        // Projection mode control
+                                        return DropdownSelector<ProjectionMode>(
+                                          borderRadius: const Radius.circular(20),
+                                          value: projectionMode,
+                                          items: ProjectionMode.values,
+                                          makeItemName: (ProjectionMode value) => _getProjectionModeName(value),
+                                          onChanged: (ProjectionMode? value) {
+                                            if (value != null) {
+                                              stageMapSettings.projectionMode.setValue(value);
+                                            }
+                                          },
+                                        );
+                                      },
+                                    ),
+                                  ),
+                                ),
+                              ],
+                            ),
+                          ),
+                        ),
+                      ),
+
+                    // Real time rendering warning
+                    if (_bShowRealTimeDisabledWarning)
+                      Positioned(
+                        bottom: 0,
+                        right: 0,
+                        child: Padding(
+                          padding: EdgeInsets.all(18),
+                          child: Text(
+                            AppLocalizations.of(context)!.stageMapRealTimeWarning,
+                            textAlign: TextAlign.end,
+                            style: TextStyle(
+                              color: Colors.yellow,
+                              shadows: [BoxShadow(color: Colors.black, offset: Offset(1, 1), blurRadius: 1)],
+                            ),
+                          ),
+                        ),
+                      ),
+                  ],
+                ),
+              );
+            },
+          );
+        },
+      ),
+    );
+  }
+
+  @override
+  void initState() {
+    super.initState();
+
+    _connectionManager = Provider.of<EngineConnectionManager>(context, listen: false);
+    _startListeningToEngineConnection();
+
+    final selectedActorSettings = Provider.of<SelectedActorSettings>(context, listen: false);
+    final stageMapSettings = Provider.of<StageMapSettings>(context, listen: false);
+
+    _settingSubscriptions.addAll([
+      selectedActorSettings.displayClusterRootPath.listen(_onDisplayClusterRootPathChanged),
+      selectedActorSettings.selectedActors.listen((_) => _updateActorTransformRegistration()),
+      stageMapSettings.bIsInTransformMode.listen(_onTransformModeChanged),
+    ]);
+
+    _actorManager = Provider.of<UnrealActorManager>(context, listen: false);
+    _actorManager.currentRootActorLightCards.addListener(guardedRefresh);
+    _actorManager.watchExistingSubscriptions(_onActorUpdate);
+
+    _propertyManager = Provider.of<UnrealPropertyManager>(context, listen: false);
+    _updateActorTransformRegistration();
+
+    _previewRenderManager = Provider.of<PreviewRenderManager>(context, listen: false);
+    _previewRenderManager.addConsumer(this);
+
+    _mapTransform.addListener(_saveMapTransformToUserSettings);
+    _previewImageStreamListener = ImageStreamListener(_onPreviewImageLoaded, onError: _onPreviewImageError);
+  }
+
+  @override
+  void dispose() {
+    for (final StreamSubscription subscription in _settingSubscriptions) {
+      subscription.cancel();
+    }
+
+    _actorManager.currentRootActorLightCards.removeListener(guardedRefresh);
+    _actorManager.stopWatchingExistingSubscriptions(_onActorUpdate);
+
+    _evictImage();
+
+    _endPropertyTransaction();
+    _stopListeningToEngineConnection();
+    _unregisterAllActorTransformProperties();
+
+    _previewRenderManager.removeConsumer(this);
+
+    _mapTransform.removeListener(_saveMapTransformToUserSettings);
+    _previewImageStream?.removeListener(_previewImageStreamListener);
+
+    _engineUpdateTimer?.cancel();
+
+    super.dispose();
+  }
+
+  @override
+  void onPreviewRenderCompleted({
+    required bool bIsForLatestRequest,
+    required Size resolution,
+    required Uint8List imageData,
+    required List<PreviewActorData> actorPositions,
+    required bool bIsRealTimeDisabled,
+  }) {
+    // Update the image data and wait for it to finish decoding asynchronously
+    if (imageData.isNotEmpty) {
+      _evictImage();
+
+      _previewImage = Image.memory(
+        imageData,
+        key: _mapKey,
+        gaplessPlayback: true,
+        fit: BoxFit.fitHeight,
+      );
+
+      _previewImageStream = _previewImage!.image.resolve(const ImageConfiguration());
+      _previewImageStream!.addListener(_previewImageStreamListener);
+      _pendingPreviewSize = resolution;
+    }
+
+    if (bIsRealTimeDisabled != _bShowRealTimeDisabledWarning) {
+      setState(() {
+        _bShowRealTimeDisabledWarning = bIsRealTimeDisabled;
+      });
+    }
+
+    // Save the latest actor positions so we can display them as soon as the image finishes decoding.
+    // If we were to update them right now, they would be out of sync with the old image.
+    _pendingActorPositions = actorPositions;
+
+    if (bIsForLatestRequest) {
+      // This render takes into account all drag operations so far, so we can now safely override the position of pins
+      // placed directly by the user without them jumping back to an old position
+      _recentDirectDragPins.clear();
+    }
+  }
+
+  /// Zoom and pan the map to focus on a specific [actor]'s pin location.
+  void focusActor(UnrealObject actor) {
+    final _MapActorData? actorData = _mapActorDataByPath[actor.path];
+    if (actorData == null) {
+      return;
+    }
+
+    _focusAndZoom(actorData.position, _maxScale);
+  }
+
+  /// Check whether an [actor] can be focused using [focusActor].
+  bool getCanFocusActor(UnrealObject actor) {
+    return _mapActorDataByPath.containsKey(actor.path);
+  }
+
+  /// Start listening for messages from the engine connection manager.
+  void _startListeningToEngineConnection() {
+    _connectionManager
+      ..registerMessageListener('ActorDragCancelled', _onActorDragCancelled)
+      ..registerConnectionStateListener(_onEngineConnectionStateChanged);
+
+    _onConnected();
+  }
+
+  /// Stop listening for messages from the engine connection manager.
+  void _stopListeningToEngineConnection() {
+    _connectionManager
+      ..unregisterMessageListener('ActorDragCancelled', _onActorDragCancelled)
+      ..unregisterConnectionStateListener(_onEngineConnectionStateChanged);
+  }
+
+  /// Called when the engine's connection state changes.
+  void _onEngineConnectionStateChanged(EngineConnectionState state) {
+    switch (state) {
+      case EngineConnectionState.connected:
+        _onConnected();
+        break;
+
+      case EngineConnectionState.disconnected:
+        _onDisconnected();
+        break;
+
+      default:
+        break;
+    }
+  }
+
+  /// Called when we connect to the engine.
+  void _onConnected() {
+    _engineUpdateTimer ??= Timer.periodic(_engineUpdateRate, _onEngineUpdateTick);
+  }
+
+  /// Called when we disconnect from the engine.
+  void _onDisconnected() {
+    _resetRootActorSpecificData();
+  }
+
+  /// Reset any data tied to the currently selected root actor.
+  void _resetRootActorSpecificData() {
+    _endPropertyTransaction();
+    _evictImage();
+
+    _previewImage = null;
+    _previewSize = null;
+    _pendingPreviewSize = Size.zero;
+    _pendingActorPositions.clear();
+
+    _mapActorDataByPath.clear();
+    _bShowRealTimeDisabledWarning = false;
+
+    setState(() {});
+  }
+
+  /// Evict the current image from cache if it exists.
+  void _evictImage() {
+    if (_previewImage != null) {
+      _previewImage?.image.evict();
+    }
+  }
+
+  /// Update the cached constraints size and the corresponding preview renderer size.
+  void _updateConstraints(BoxConstraints constraints) {
+    // Update preview size based on the preview area's size
+    Size newSize = Size(constraints.maxWidth, constraints.maxHeight);
+    if (newSize != _constraintSize) {
+      _constraintSize = newSize;
+
+      if (!widget.bForceShowFullMap) {
+        // Map transform settings are stored in an aspect ratio-independent way, which means
+        // we need to reload them whenever we have a potentially new aspect ratio
+        _setMapTransformFromUserSettings(bShouldForceBuild: false);
+      }
+
+      // Update preview resolution
+      if (_constraintSize.height > 0 && _constraintSize.width > 0) {
+        final double newWidth = math.sqrt(_previewResolutionPixels * _containerAspectRatio);
+        _desiredPreviewSize = Size(newWidth, newWidth / _containerAspectRatio);
+      } else {
+        _desiredPreviewSize = Size.zero;
+      }
+
+      _previewRenderManager.setDesiredResolution(_desiredPreviewSize);
+    }
+  }
+
+  /// Called when the transform mode setting changes.
+  void _onTransformModeChanged(bool bIsInTransformMode) {
+    // Register/unregister for properties necessary to transform a light card
+    _updateActorTransformRegistration();
+  }
+
+  /// Called when the root display cluster path changes.
+  void _onDisplayClusterRootPathChanged(String newPath) {
+    _resetRootActorSpecificData();
+  }
+
+  /// Called when a preview image finishes loading from memory.
+  void _onPreviewImageLoaded(ImageInfo image, bool synchronousCall) {
+    if (!mounted) {
+      return;
+    }
+
+    setState(() {
+      _previewSize = _pendingPreviewSize;
+
+      final missingActorPaths = Set<String>.from(_mapActorDataByPath.keys);
+
+      // Update positions for all received paths
+      for (final PreviewActorData actorData in _pendingActorPositions) {
+        missingActorPaths.remove(actorData.path);
+
+        if (_recentDirectDragPins.contains(actorData.path)) {
+          // These actors were manually placed by the user after this render was generated, so don't move them
+          continue;
+        }
+
+        _getOrCreateMapActorData(actorData.path).position = actorData.position;
+      }
+
+      // Remove any actors that are no longer reported by the engine
+      for (String actorPath in missingActorPaths) {
+        _unregisterActorTransformProperties(actorPath);
+        _mapActorDataByPath.remove(actorPath);
+      }
+    });
+
+    _pendingActorPositions.clear();
+  }
+
+  /// Called when the preview image failed to decode.
+  void _onPreviewImageError(Object exception, StackTrace? stackTrace) {
+    _previewImageStream?.removeListener(_previewImageStreamListener);
+  }
+
+  /// Called when Unreal cancels a actor drag that was still in progress.
+  void _onActorDragCancelled(dynamic message) {
+    if (_draggedPinPath == null || message['RendererId'] != _previewRenderManager.rendererId) {
+      return;
+    }
+
+    // Re-start the drag that was in progress
+    _startPinDrag(_draggedPinPath!);
+  }
+
+  /// Called each tick, when we're ready to send updates back to the engine.
+  void _onEngineUpdateTick(Timer timer) {
+    if (_connectionManager.connectionState != EngineConnectionState.connected ||
+        _previewRenderManager.rendererId == null) {
+      return;
+    }
+
+    if (_tickDragPosition != null) {
+      // Send the latest actor drag position
+      _previewRenderManager.updateActorDrag(_tickDragPosition!);
+      _tickDragPosition = null;
+    }
+  }
+
+  /// Tell the engine that we've started dragging a pin.
+  /// If a [primaryPinPath] is provided, that pin will be the one around which other pins' locations will rotate to
+  /// match the wall's angle.
+  /// If no path is provided, the nearest pin to the pointer will be used as the primary, so the [pointerPosition] must
+  /// be provided instead.
+  void _startPinDrag(String? primaryPinPath, [Offset? pointerPosition]) {
+    if (primaryPinPath == null && pointerPosition == null) {
+      throw Exception('_startPinDrag must be called with at least one of primaryPinPath and pointerPosition');
+    }
+
+    final stageMapSettings = Provider.of<StageMapSettings>(context, listen: false);
+    final selectedActorSettings = Provider.of<SelectedActorSettings>(context, listen: false);
+    late final List<String> draggedActors;
+
+    if (selectedActorSettings.bIsInMultiSelectMode.getValue() || stageMapSettings.bIsInTransformMode.getValue()) {
+      // In multi-select mode (or single-select transform mode), drag the currently selected actors, using either the
+      // primary pin or the nearest selected pin to the pointer position
+      draggedActors = selectedActorSettings.selectedActors.getValue().toList();
+
+      if (primaryPinPath == null) {
+        // There's no primary pin, so find the nearest selected pin
+
+        if (draggedActors.isEmpty) {
+          // There are no selected pins, so a drag wouldn't do anything anyway
+          return;
+        }
+
+        final RenderBox? mapRenderBox = _mapKey.currentContext?.findRenderObject() as RenderBox;
+        if (mapRenderBox == null) {
+          return;
+        }
+
+        // Find the selected pin nearest the pointer and use it as the primary pin
+        double nearestDistanceSquared = double.infinity;
+        Offset nearestPinScreenPosition = Offset.zero;
+
+        for (final String pinPath in draggedActors) {
+          final _MapActorData? mapActorData = _mapActorDataByPath[pinPath];
+          if (mapActorData == null) {
+            continue;
+          }
+
+          final Offset pinScreenPosition = mapRenderBox.localToGlobal(mapActorData.position.scale(
+            mapRenderBox.size.width,
+            mapRenderBox.size.height,
+          ));
+
+          final pinDistance = (pinScreenPosition - pointerPosition!).distanceSquared;
+          if (pinDistance < nearestDistanceSquared) {
+            nearestDistanceSquared = pinDistance;
+            nearestPinScreenPosition = pinScreenPosition;
+            primaryPinPath = pinPath;
+          }
+        }
+
+        assert(primaryPinPath != null);
+
+        _dragPointerOffset = nearestPinScreenPosition - pointerPosition!;
+      } else {
+        // We have a primary pin already
+
+        // If the primary pin isn't selected, drag it along with the selection as well
+        if (!draggedActors.contains(primaryPinPath)) {
+          draggedActors.add(primaryPinPath);
+        }
+
+        _dragPointerOffset = Offset.zero;
+      }
+    } else {
+      // We're in single-select map mode, so only the primary pin is a valid drag target
+      if (primaryPinPath == null) {
+        return;
+      }
+
+      // If the primary pin isn't selected, switch the selection to it
+      if (!selectedActorSettings.isActorSelected(primaryPinPath)) {
+        selectedActorSettings.selectActor(primaryPinPath);
+      }
+
+      draggedActors = [primaryPinPath];
+      _dragPointerOffset = Offset.zero;
+    }
+
+    _previewRenderManager.beginActorDrag(
+      actors: draggedActors,
+      primaryActor: primaryPinPath,
+    );
+
+    setState(() {
+      _draggedPinPath = primaryPinPath;
+      _tickDragPosition = null;
+      _gestureState = _GestureState.transformDragging;
+    });
+  }
+
+  /// Update our stored light card position based on the latest drag position.
+  void _updatePinDrag(Offset screenPosition) {
+    if (_gestureState != _GestureState.transformDragging) {
+      return;
+    }
+
+    final Offset? mapPosition = _dragPin(screenPosition);
+    if (mapPosition == null) {
+      return;
+    }
+
+    _tickDragPosition = mapPosition;
+  }
+
+  /// Tell the engine that we're done dragging pins.
+  void _endPinDrag(Offset? screenPosition) {
+    _gestureState = _GestureState.none;
+
+    if (_draggedPinPath == null) {
+      return;
+    }
+
+    // Use the latest drag position. If that fails, fall back to pin's last known position.
+    Offset? mapPosition;
+    if (screenPosition != null) {
+      mapPosition = _dragPin(screenPosition);
+    }
+
+    if (mapPosition == null) {
+      _mapActorDataByPath[_draggedPinPath]?.position;
+    }
+
+    setState(() {
+      _tickDragPosition = null;
+      _draggedPinPath = null;
+    });
+
+    if (mapPosition == null) {
+      return;
+    }
+
+    _previewRenderManager.endActorDrag(mapPosition);
+  }
+
+  /// Check whether a pointer event is a right-click from a mouse.
+  bool _isRightClick(PointerEvent event) {
+    return event.kind == PointerDeviceKind.mouse && event.buttons == 2;
+  }
+
+  /// Called when a pointer is pressed on the map.
+  void _onPointerDown(PointerDownEvent event) {
+    if (!_isRightClick(event)) {
+      return;
+    }
+
+    // Simulate a two-finger gesture so we can easily debug on a desktop
+    _onScaleStart(ScaleStartDetails(
+      focalPoint: event.position,
+      localFocalPoint: event.localPosition,
+      pointerCount: 2,
+    ));
+  }
+
+  /// Called when a pointer is released on the map.
+  void _onPointerUp(PointerUpEvent event) {
+    if (!_isRightClick(event)) {
+      return;
+    }
+
+    // Simulate a two-finger gesture so we can easily debug on a desktop
+    _onScaleEnd(ScaleEndDetails(
+      pointerCount: 2,
+    ));
+  }
+
+  /// Called when a pointer leaves the map area.
+  void _onPointerCancel(PointerCancelEvent event) {
+    if (!_isRightClick(event)) {
+      return;
+    }
+
+    // Simulate a two-finger gesture so we can easily debug on a desktop
+    _onScaleEnd(ScaleEndDetails(
+      pointerCount: 2,
+    ));
+  }
+
+  /// Called when a pointer moves over the map.
+  void _onPointerMove(PointerMoveEvent event) {
+    if (!_isRightClick(event)) {
+      return;
+    }
+
+    // Simulate a two-finger gesture so we can easily debug on a desktop
+    _onScaleUpdate(ScaleUpdateDetails(
+      focalPoint: event.position,
+      localFocalPoint: event.localPosition,
+      scale: 1.0,
+      pointerCount: 2,
+    ));
+  }
+
+  /// Called when a pointer signal event happens on the map.
+  void _onPointerSignal(PointerSignalEvent event) {
+    if (event is! PointerScrollEvent) {
+      return;
+    }
+
+    // Simulate a two-finger gesture so we can easily debug on a desktop
+    _onScaleStart(ScaleStartDetails(
+      focalPoint: event.position,
+      localFocalPoint: event.localPosition,
+      pointerCount: 2,
+    ));
+
+    // Implementation based on InteractiveViewer to get a natural-feeling scroll.
+    final double scaleChange = math.exp(-event.scrollDelta.dy / _scrollScaleFactor);
+
+    _onScaleUpdate(ScaleUpdateDetails(
+      focalPoint: event.position,
+      localFocalPoint: event.localPosition,
+      scale: scaleChange,
+      pointerCount: 2,
+    ));
+
+    _onScaleEnd(ScaleEndDetails(pointerCount: 2));
+  }
+
+  /// Called when the user starts a scale/rotate gesture on the map.
+  void _onScaleStart(ScaleStartDetails details) {
+    _gestureReferenceScenePoint = _mapTransform.toScene(details.localFocalPoint);
+    _gestureStartPosition = details.focalPoint;
+    _lastScaleUpdateDetails = ScaleUpdateDetails(
+      pointerCount: details.pointerCount,
+      focalPoint: details.focalPoint,
+      localFocalPoint: details.localFocalPoint,
+    );
+
+    final stageMapSettings = Provider.of<StageMapSettings>(context, listen: false);
+    if (stageMapSettings.bIsInTransformMode.getValue()) {
+      if (_gestureState != _GestureState.none) {
+        return;
+      }
+
+      if (details.pointerCount == 1) {
+        _onTransformPanStart(details);
+      } else if (details.pointerCount >= 2) {
+        _onTransformScaleStart();
+      }
+    } else if (details.pointerCount > 1 || stageMapSettings.projectionMode.getValue() == ProjectionMode.uv) {
+      // Start a camera track and zoom if we're using a scale gesture or just dragging in UV mode (since it can't
+      // tumble)
+      _onCameraTrackAndZoomStart(details);
+    }
+  }
+
+  /// Called when the user continues a scale/rotate/pan gesture on the map.
+  void _onScaleUpdate(ScaleUpdateDetails details) {
+    final stageMapSettings = Provider.of<StageMapSettings>(context, listen: false);
+    if (stageMapSettings.bIsInTransformMode.getValue()) {
+      switch (_gestureState) {
+        case _GestureState.transformScaling:
+          _onTransformScaleUpdate(details);
+          break;
+
+        case _GestureState.transformDragging:
+          _onTransformPanUpdate(details);
+          break;
+
+        default:
+          break;
+      }
+    } else if (details.pointerCount == 1 && stageMapSettings.projectionMode.getValue() != ProjectionMode.uv) {
+      // We don't tumble in UV mode, so skip to track and zoom in that case
+      _onCameraTumbleUpdate(details);
+    } else {
+      _onCameraTrackAndZoomUpdate(details);
+    }
+
+    _lastScaleUpdateDetails = details;
+  }
+
+  /// Called when the user ends a scale/rotate gesture on the map.
+  void _onScaleEnd(ScaleEndDetails details) {
+    if (Provider.of<StageMapSettings>(context, listen: false).bIsInTransformMode.getValue()) {
+      switch (_gestureState) {
+        case _GestureState.transformScaling:
+          _onTransformScaleEnd();
+          break;
+
+        case _GestureState.transformDragging:
+          _onTransformPanEnd();
+          break;
+
+        default:
+          break;
+      }
+    }
+
+    _gestureState = _GestureState.none;
+  }
+
+  /// Called when the user performs a long press gesture on the map.
+  void _onLongPressStart(LongPressStartDetails details) {
+    DropDownListMenu.showAtPosition(
+      pivotPosition: details.globalPosition,
+      builder: (context) => PlaceActorDropDownMenu(
+        actorMapPosition: _globalToMapPosition(details.globalPosition),
+        bIsFromLongPress: true,
+      ),
+    );
+  }
+
+  /// Called when user the continues to perform a drag gesture that should tumble the camera.
+  void _onCameraTumbleUpdate(ScaleUpdateDetails details) {
+    if (_gestureState == _GestureState.none) {
+      // Start the tumble operation if we've moved enough to be confident that it's what the user wants
+      final double focalPointDistance = (details.focalPoint - _gestureStartPosition).distance;
+      if (focalPointDistance > _tumbleTouchSlop) {
+        _gestureState = _GestureState.cameraTumbling;
+      }
+    }
+
+    if (_gestureState != _GestureState.cameraTumbling) {
+      return;
+    }
+
+    final stageMapSettings = Provider.of<StageMapSettings>(context, listen: false);
+
+    vec.Vector2 cameraAngle = stageMapSettings.cameraAngle.getValue();
+    cameraAngle.x += details.focalPointDelta.dx * _tumbleSensitivity.dx;
+    cameraAngle.y = (cameraAngle.y - details.focalPointDelta.dy * _tumbleSensitivity.dy).clamp(-90.0, 90.0);
+
+    stageMapSettings.cameraAngle.setValue(cameraAngle);
+  }
+
+  /// Called when the user starts a scale gesture that should track and zoom the camera.
+  void _onCameraTrackAndZoomStart(ScaleStartDetails details) {
+    if (_gestureState != _GestureState.none) {
+      return;
+    }
+
+    _gestureState = _GestureState.cameraTrackingAndZooming;
+    _gestureStartScale = _mapTransform.value.getMaxScaleOnAxis();
+  }
+
+  /// Called when the user continues scale gesture that should track and zoom the camera.
+  void _onCameraTrackAndZoomUpdate(ScaleUpdateDetails details) {
+    if (_gestureState != _GestureState.cameraTrackingAndZooming) {
+      return;
+    }
+
+    // Apply scale changes
+    final double currentScale = _mapTransform.value.getMaxScaleOnAxis();
+    assert(currentScale != 0.0);
+
+    final double desiredScale = _gestureStartScale * details.scale;
+    final double deltaScale = desiredScale / currentScale;
+
+    _scaleMap(deltaScale);
+
+    // After scaling, apply the translation change. We want to do this even if the translation is zero, because the new
+    // scale may have pushed the viewport out of bounds.
+    final Offset scenePoint = _mapTransform.toScene(details.localFocalPoint);
+    final Offset translation = scenePoint - _gestureReferenceScenePoint;
+
+    _translateMap(translation);
+
+    _gestureReferenceScenePoint = _mapTransform.toScene(details.localFocalPoint);
+  }
+
+  /// Called when the user starts a scale/rotate gesture in transform mode.
+  void _onTransformScaleStart() {
+    if (_bIsWaitingForExpose) {
+      return;
+    }
+
+    if (_propertyManager.beginTransaction(AppLocalizations.of(context)!.transactionTransformStageActors)) {
+      _gestureState = _GestureState.transformScaling;
+    }
+  }
+
+  /// Called when the user continues a scale/rotate gesture in transform mode.
+  void _onTransformScaleUpdate(ScaleUpdateDetails details) {
+    if (_bIsWaitingForExpose) {
+      return;
+    }
+
+    final double deltaSpin = -vec.degrees(details.rotation - _lastScaleUpdateDetails.rotation);
+    final double deltaScalar = (details.scale - _lastScaleUpdateDetails.scale) * 2;
+
+    for (final _MapActorData mapActorData in _mapActorDataByPath.values) {
+      if (!mapActorData.bArePropsRegistered) {
+        continue;
+      }
+
+      // Multiply by the current scale so we can use an add operation, which is more stable than multiplying repeatedly
+      final vec.Vector2 deltaScale =
+          _propertyManager.getTrackedPropertyValue(mapActorData.scalePropertyId) * deltaScalar;
+
+      _propertyManager.modifyTrackedPropertyValue(mapActorData.spinPropertyId, const AddOperation(),
+          deltaValue: deltaSpin);
+      _propertyManager.modifyTrackedPropertyValue(mapActorData.scalePropertyId, const AddOperation(),
+          deltaValue: deltaScale);
+    }
+  }
+
+  /// Called when the user ends a scale/rotate gesture in transform mode.
+  void _onTransformScaleEnd() {
+    _endPropertyTransaction();
+  }
+
+  /// Called when the user starts a pan gesture in transform mode.
+  void _onTransformPanStart(ScaleStartDetails details) {
+    if (_bIsWaitingForExpose) {
+      return;
+    }
+
+    _startPinDrag(null, details.focalPoint);
+  }
+
+  /// Called when the user continues a pan gesture in transform mode.
+  void _onTransformPanUpdate(ScaleUpdateDetails details) {
+    if (_bIsWaitingForExpose) {
+      return;
+    }
+
+    _updatePinDrag(details.focalPoint);
+  }
+
+  /// Called when the user ends a pan gesture in transform mode.
+  void _onTransformPanEnd() {
+    if (_bIsWaitingForExpose) {
+      return;
+    }
+
+    _endPinDrag(_lastScaleUpdateDetails.focalPoint);
+  }
+
+  /// End the current actor property transaction. Note that this is unnecessary when changing actors' positions,
+  /// as those transactions are handled internally by the engine.
+  void _endPropertyTransaction() {
+    if (_gestureState == _GestureState.none) {
+      return;
+    }
+
+    _propertyManager.endTransaction();
+  }
+
+  /// Calculate the viewport area, i.e. the area of the map currently filling the container, when the map is
+  /// using the given transformation [matrix].
+  Rect _calculateViewport(Matrix4 matrix) {
+    final Rect containerRect = _containerRect;
+
+    // Viewport is always the inverse of the map's actually transform (e.g. when we zoom in, the viewport shrinks)
+    final Matrix4 viewportMatrix = matrix.clone()..invert();
+
+    final topLeft = viewportMatrix.transform3(vec.Vector3(
+      containerRect.left,
+      containerRect.top,
+      0,
+    ));
+    final bottomRight = viewportMatrix.transform3(vec.Vector3(
+      containerRect.right,
+      containerRect.bottom,
+      0,
+    ));
+
+    return Offset(topLeft.x, topLeft.y) & Size(bottomRight.x - topLeft.x, bottomRight.y - topLeft.y);
+  }
+
+  /// Calculate the amount by which the [inner] rectangle sits outside of the [outer] rectangle.
+  Offset _calculateRectExcess(Rect inner, Rect outer) {
+    final double rightExcess = (inner.right > outer.right) ? (inner.right - outer.right) : 0.0;
+    final double leftExcess = (inner.left < outer.left) ? (inner.left - outer.left) : 0.0;
+    final double topExcess = (inner.top < outer.top) ? (inner.top - outer.top) : 0.0;
+    final double bottomExcess = (inner.bottom > outer.bottom) ? (inner.bottom - outer.bottom) : 0.0;
+
+    return Offset(
+      (rightExcess.abs() > leftExcess.abs()) ? rightExcess : leftExcess,
+      (topExcess.abs() > bottomExcess.abs()) ? topExcess : bottomExcess,
+    );
+  }
+
+  /// Translate the map using the given [translation] offset while keeping the viewport in valid bounds.
+  /// If [bShouldForceBuild] is true, force the widget to rebuild after translation is applied.
+  void _translateMap(Offset translation, {bool bShouldForceBuild = true}) {
+    final Matrix4 newMatrix = _mapTransform.value.clone()
+      ..translate(
+        translation.dx,
+        translation.dy,
+      );
+
+    final Rect containerRect = _containerRect;
+    final Rect viewport = _calculateViewport(newMatrix);
+    final Offset excess = _calculateRectExcess(viewport, containerRect);
+
+    const double maxExcess = 0.00001;
+    if (excess.distance < maxExcess) {
+      _mapTransform.value = newMatrix;
+      if (bShouldForceBuild) {
+        setState(() {});
+      }
+      return;
+    }
+
+    // New viewport is out of bounds, so adjust the matrix accordingly
+    newMatrix.translate(excess.dx, excess.dy);
+
+    // Check if we're still out of bounds, in which case the viewport doesn't fit and we can't translate at all
+    final Rect correctedViewport = _calculateViewport(newMatrix);
+    final Offset correctedExcess = _calculateRectExcess(correctedViewport, containerRect);
+
+    if (correctedExcess.distance < maxExcess) {
+      _mapTransform.value = newMatrix;
+      if (bShouldForceBuild) {
+        setState(() {});
+      }
+      return;
+    }
+  }
+
+  /// Scale the map by the given [deltaScale] while keeping the scale within valid bounds.
+  void _scaleMap(double deltaScale) {
+    assert(deltaScale != 0.0);
+
+    final double currentScale = _mapTransform.value.getMaxScaleOnAxis();
+    assert(currentScale != 0.0);
+
+    final double newScale = (currentScale * deltaScale).clamp(_minScale, _maxScale);
+    final double correctedDeltaScale = newScale / currentScale;
+
+    if (correctedDeltaScale != 1.0) {
+      setState(() {
+        _mapTransform.value = _mapTransform.value.clone()..scale(correctedDeltaScale);
+      });
+    }
+  }
+
+  /// Save the map's transform to the user settings.
+  void _saveMapTransformToUserSettings() {
+    final Rect containerRect = _containerRect;
+
+    if (containerRect.longestSide == 0) {
+      return;
+    }
+
+    final Offset containerCenter = containerRect.center;
+    final Offset focalPoint = _mapTransform.toScene(containerCenter).scale(
+          1 / containerRect.width,
+          1 / containerRect.height,
+        );
+
+    // Focal point setting assumes a 1:1 aspect ratio, so scale horizontal coordinate relative to center position
+    final Offset aspectScaledFocalPoint = Offset(
+      0.5 + (focalPoint.dx - 0.5) * _containerAspectRatio,
+      focalPoint.dy,
+    );
+
+    final stageMapSettings = Provider.of<StageMapSettings>(context, listen: false);
+    stageMapSettings.focalPoint.setValue(aspectScaledFocalPoint);
+    stageMapSettings.zoomLevel.setValue(_mapTransform.value.getMaxScaleOnAxis());
+  }
+
+  /// Update the stage map's transform based on the saved user settings.
+  /// If [bShouldForceBuild] is true, force the widget to rebuild after translation is applied.
+  void _setMapTransformFromUserSettings({bool bShouldForceBuild = true}) {
+    final stageMapSettings = Provider.of<StageMapSettings>(context, listen: false);
+
+    final double zoomLevel = stageMapSettings.zoomLevel.getValue();
+    final Offset focalPoint = stageMapSettings.focalPoint.getValue();
+
+    // Focal point setting assumes a 1:1 aspect ratio, so scale horizontal coordinate relative to center
+    final Offset aspectScaledFocalPoint = Offset(
+      0.5 + (focalPoint.dx - 0.5) / _containerAspectRatio,
+      focalPoint.dy,
+    );
+
+    _focusAndZoom(aspectScaledFocalPoint, zoomLevel, bShouldForceBuild: bShouldForceBuild);
+  }
+
+  /// Change the map to a specific zoom level and centered focal point, maintaining correct limits.
+  /// If [bShouldForceBuild] is true, force the widget to rebuild after translation is applied.
+  void _focusAndZoom(Offset focalPoint, double zoomLevel, {bool bShouldForceBuild = true}) {
+    final Rect containerRect = _containerRect;
+
+    if (containerRect.longestSide == 0) {
+      return;
+    }
+
+    // Zoom in
+    _mapTransform.value = Matrix4.diagonal3(vec.Vector3(zoomLevel, zoomLevel, zoomLevel));
+
+    // Translate so the focal point is centered within the container
+    final Offset focalMapPosition = focalPoint.scale(containerRect.width, containerRect.height);
+    final Offset center = _containerRect.center.scale(1 / zoomLevel, 1 / zoomLevel);
+
+    _translateMap(center - focalMapPosition, bShouldForceBuild: bShouldForceBuild);
+  }
+
+  /// Get the stored data for a actor by its path, or create it if it doesn't exist.
+  _MapActorData _getOrCreateMapActorData(String path) {
+    final _MapActorData? mapActorData = _mapActorDataByPath[path];
+    if (mapActorData != null) {
+      return mapActorData;
+    }
+
+    return _mapActorDataByPath[path] = _MapActorData();
+  }
+
+  /// Update which light cards' properties are registered to match the current selected list.
+  void _updateActorTransformRegistration() {
+    final stageMapSettings = Provider.of<StageMapSettings>(context, listen: false);
+    if (!stageMapSettings.bIsInTransformMode.getValue()) {
+      // If we're not transforming actors, we don't care about any of their properties
+      _unregisterAllActorTransformProperties();
+      return;
+    }
+
+    bool bAnyChanges = false;
+
+    // Unregister actors that are no longer selected
+    final Set<String> selectedActors =
+        Provider.of<SelectedActorSettings>(context, listen: false).selectedActors.getValue();
+    final List<String> registeredActorPaths =
+        _mapActorDataByPath.keys.where((path) => _mapActorDataByPath[path]!.bArePropsRegistered).toList();
+
+    for (final String actorPath in registeredActorPaths) {
+      if (selectedActors.contains(actorPath)) {
+        continue;
+      }
+
+      bAnyChanges |= _unregisterActorTransformProperties(actorPath);
+    }
+
+    // Register new actors
+    for (String actorPath in selectedActors) {
+      if (_mapActorDataByPath[actorPath]?.bArePropsRegistered == true) {
+        continue;
+      }
+
+      bAnyChanges |= _registerActorTransformProperties(actorPath);
+    }
+
+    if (bAnyChanges) {
+      // Tracked properties were changed, so get all the property IDs and wait for them to be ready
+      setState(() {
+        _bIsWaitingForExpose = true;
+      });
+
+      final Iterable<Future> waitFutures =
+          _mapActorDataByPath.values.map((mapActorData) => mapActorData.waitForExpose(_propertyManager));
+      Future.wait(waitFutures).then((_) {
+        if (!mounted) {
+          return;
+        }
+
+        setState(() {
+          _bIsWaitingForExpose = false;
+        });
+      }).onError((error, stackTrace) {
+        if (!mounted) {
+          return;
+        }
+
+        _log.severe('Error waiting for properties to enter transform mode', error, stackTrace);
+
+        stageMapSettings.bIsInTransformMode.setValue(false);
+
+        setState(() {
+          _bIsWaitingForExpose = false;
+        });
+      });
+    }
+  }
+
+  /// Register a light card's properties needed for transform controls with the property manager.
+  bool _registerActorTransformProperties(String actorPath) {
+    if (_mapActorDataByPath[actorPath]?.bArePropsRegistered == true) {
+      // We're already tracking this actor's properties
+      return false;
+    }
+
+    final UnrealObject? actor = _actorManager.getActorAtPath(actorPath);
+
+    if (actor == null || !actor.isAny(controllableClassNames)) {
+      return false;
+    }
+
+    final _MapActorData mapActorData = _getOrCreateMapActorData(actorPath);
+    mapActorData.registerProperties(_propertyManager, actor);
+
+    return true;
+  }
+
+  /// Unregister a light card's properties needed for transform controls with the property manager.
+  bool _unregisterActorTransformProperties(String actorPath) {
+    final _MapActorData? mapActorData = _mapActorDataByPath[actorPath];
+    if (mapActorData == null || !mapActorData.bArePropsRegistered) {
+      return false;
+    }
+
+    mapActorData.unregisterProperties(_propertyManager);
+    return true;
+  }
+
+  /// Unregister all tracked properties for all light cards.
+  void _unregisterAllActorTransformProperties() {
+    for (final String actorPath in _mapActorDataByPath.keys) {
+      _unregisterActorTransformProperties(actorPath);
+    }
+  }
+
+  /// Calculate a pin's offset within this widget given its pixel [position] relative to the map's top left point and
+  /// the maximum scale of the map when zoomed all the way out ([maxMapScale]).
+  Offset _getPinOffset(Offset position) {
+    final Size mapSize = _previewSize ?? _desiredPreviewSize;
+
+    // How much the map will be scaled down to fit the container
+    final double maxMapScale = _constraintSize.height / mapSize.height;
+
+    // How much the map will be offset left/right to be centered in the container
+    final double mapXOffset = (_constraintSize.width - mapSize.width * maxMapScale) / 2;
+
+    final transformed = vec.Vector3(
+      position.dx * maxMapScale * mapSize.width + mapXOffset,
+      position.dy * maxMapScale * mapSize.height,
+      0,
+    );
+    transformed.applyMatrix4(_mapTransform.value);
+
+    return Offset(transformed.x, transformed.y);
+  }
+
+  /// Drag a pin based on a new [screenPosition] and return the drag position that should be reported to the engine over
+  /// WebSocket (or return null if dragging the pin is invalid).
+  /// If the map is in object transform mode, the pin's actual location will not be updated (since we need to wait for
+  /// the engine to tell us its final position).
+  Offset? _dragPin(Offset screenPosition) {
+    if (_draggedPinPath == null) {
+      return null;
+    }
+
+    screenPosition += _dragPointerOffset;
+
+    final Offset? mapPosition = _globalToMapPosition(screenPosition);
+    if (mapPosition == null) {
+      return null;
+    }
+
+    final stageMapSettings = Provider.of<StageMapSettings>(context, listen: false);
+    if (!stageMapSettings.bIsInTransformMode.getValue()) {
+      // In this mode, the user is directly dragging a specific pin, so we want to immediately update its position and
+      // keep it there until the engine sends a new enough render. Otherwise, the pin will snap back to its old position
+      // if the user releases it before the engine sends a new render.
+      // Note that we only do this for the primary pin (_draggedPinPath), as any other multi-selected pins are moved
+      // indirectly in the engine and we can't predict their final positions.
+      _getOrCreateMapActorData(_draggedPinPath!).position = mapPosition;
+      _recentDirectDragPins.add(_draggedPinPath!);
+    }
+
+    return mapPosition;
+  }
+
+  /// Returns whether a actor should be forced to show its drag feedback animation.
+  /// Normally, only the directly pressed pin would show its drag feedback animation, but this lets us trigger it
+  /// on other pins if we're dragging a group with multi-select.
+  bool _shouldForcePinDragAnimation(String actorPath) {
+    if (!Provider.of<StageMapSettings>(context, listen: false).bIsInTransformMode.getValue() &&
+        actorPath == _draggedPinPath) {
+      // When dragging a actor directly, let it control its own drag animation
+      return false;
+    }
+
+    return _gestureState == _GestureState.transformDragging &&
+        Provider.of<SelectedActorSettings>(context, listen: false).isActorSelected(actorPath);
+  }
+
+  /// Get the name to display for each projection mode
+  String _getProjectionModeName(ProjectionMode mode) {
+    switch (mode) {
+      case ProjectionMode.azimuthal:
+        return AppLocalizations.of(context)!.projectionModeDome;
+
+      case ProjectionMode.orthographic:
+        return AppLocalizations.of(context)!.projectionModeOrthographic;
+
+      case ProjectionMode.perspective:
+        return AppLocalizations.of(context)!.projectionModePerspective;
+
+      case ProjectionMode.uv:
+        return AppLocalizations.of(context)!.projectionModeUV;
+    }
+  }
+
+  /// Convert a position from global screen space to a normalized coordinate within the preview map.
+  Offset? _globalToMapPosition(Offset globalPosition) {
+    final RenderBox? mapRenderBox = _mapKey.currentContext?.findRenderObject() as RenderBox;
+    if (mapRenderBox == null) {
+      return null;
+    }
+
+    // Convert to the coordinate space of the rendered preview map
+    final Offset mapLocalPosition = mapRenderBox.globalToLocal(globalPosition);
+    return Offset(
+      (mapLocalPosition.dx / mapRenderBox.size.width).clamp(0, 1),
+      (mapLocalPosition.dy / mapRenderBox.size.height).clamp(0, 1),
+    );
+  }
+
+  /// Called when any actor is added/renamed/deleted.
+  void _onActorUpdate(ActorUpdateDetails details) {
+    if (Provider.of<StageMapSettings>(context, listen: false).bIsInTransformMode.getValue() != true) {
+      return;
+    }
+
+    for (final UnrealObject actor in details.addedActors) {
+      if (Provider.of<SelectedActorSettings>(context, listen: false).isActorSelected(actor.path)) {
+        // An actor we had selected has been added, meaning we need to expose its transform properties. This can happen
+        // if we added an actor from the app and immediately selected it before the actor manager received its creation
+        // event data.
+        _updateActorTransformRegistration();
+        return;
+      }
+    }
+  }
+}
+
+/// Data about actors shown on the map.
+class _MapActorData {
+  /// Position of the actor's pin on the map. Coordinates are in range [0, 1], normalized to the size of the preview
+  /// render.
+  Offset position = Offset.zero;
+
+  /// Property ID for tracking the actor's spin.
+  TrackedPropertyId spinPropertyId = const TrackedPropertyId.invalid();
+
+  /// Property ID for tracking the actor's scale.
+  TrackedPropertyId scalePropertyId = const TrackedPropertyId.invalid();
+
+  /// Whether this has its properties registered with the property manager.
+  bool get bArePropsRegistered => spinPropertyId.isValid;
+
+  /// Register all of this light card's properties for the given [propertyManager] and [actorPath].
+  bool registerProperties(UnrealPropertyManager propertyManager, UnrealObject actor) {
+    registerProperty(String propertyName) => propertyManager.trackProperty(
+          UnrealProperty(objectPath: actor.path, propertyName: propertyName),
+          _onManagedValueChanged,
+        );
+
+    String spinPropertyName;
+    String scalePropertyName;
+
+    if (actor.isA(lightCardClassName)) {
+      spinPropertyName = 'Spin';
+      scalePropertyName = 'Scale';
+    } else if (actor.isAny(colorCorrectWindowClassNames)) {
+      spinPropertyName = 'PositionalParams.Spin';
+      scalePropertyName = 'PositionalParams.Scale';
+    } else {
+      return false;
+    }
+
+    spinPropertyId = registerProperty(spinPropertyName);
+    scalePropertyId = registerProperty(scalePropertyName);
+
+    return true;
+  }
+
+  /// Unregister all of this light card's properties with the given [propertyManager].
+  void unregisterProperties(UnrealPropertyManager propertyManager) {
+    propertyManager.stopTrackingProperty(spinPropertyId, _onManagedValueChanged);
+    propertyManager.stopTrackingProperty(scalePropertyId, _onManagedValueChanged);
+
+    spinPropertyId = const TrackedPropertyId.invalid();
+    scalePropertyId = const TrackedPropertyId.invalid();
+  }
+
+  /// Wait until all properties for this actor have been exposed to the given [propertyManager].
+  Future waitForExpose(UnrealPropertyManager propertyManager) {
+    return Future.wait([
+      if (spinPropertyId.isValid) propertyManager.waitForProperty(spinPropertyId),
+      if (scalePropertyId.isValid) propertyManager.waitForProperty(scalePropertyId)
+    ]);
+  }
+
+  /// Stub function passed to the property manager when registering actor properties whose values we don't actually
+  /// do anything with, but needs to be passed again to unregister them.
+  void _onManagedValueChanged(dynamic newValue) {}
+}
+
+/// A draggable pin displayed at the location of a stage actor.
+class _StageMapPin extends StatefulWidget {
+  const _StageMapPin({
+    Key? key,
+    required this.actor,
+    this.bIsDraggable = true,
+    this.bIsSelectable = true,
+    this.bShowAsDragged = false,
+    this.size,
+    this.onDragStarted,
+    this.onDragUpdate,
+    this.onDragEnded,
+  }) : super(key: key);
+
+  /// The stage actor this represents.
+  final UnrealObject actor;
+
+  /// Whether this pin can be dragged.
+  final bool bIsDraggable;
+
+  /// Whether this pin can be selected and deselected.
+  final bool bIsSelectable;
+
+  /// Whether to show the drag animation regardless of whether this is currently being dragged directly.
+  final bool bShowAsDragged;
+
+  /// The size of the pin icon. If not provided, a default size will be used.
+  final double? size;
+
+  /// Callback when a pin starts to be dragged.
+  final Function()? onDragStarted;
+
+  /// Callback when a pin's position changes while being dragged. Passes the screen position of the dropped pin.
+  final Function(Offset pinPosition)? onDragUpdate;
+
+  /// Callback when a pin stops being dragged. Passes the screen position of the dropped pin.
+  final Function(Offset pinPosition)? onDragEnded;
+
+  @override
+  State<_StageMapPin> createState() => _StageMapPinState();
+}
+
+class _StageMapPinState extends State<_StageMapPin> with TickerProviderStateMixin {
+  /// Delay before a long press starts a drag on this pin.
+  static const Duration _dragDelay = Duration(milliseconds: 300);
+
+  /// Default size of the pin on the screen.
+  static const double _defaultPinSize = 96.0;
+
+  /// Y offset of the pin graphic's bottom from the actual pinned location.
+  static const double _bottomOffset = 10.0;
+
+  /// How high up the pin should move when it's being held (as a fraction of the pin's height).
+  static const double _heldHeightOffset = 0.3;
+
+  /// The size of the pin on the screen.
+  double get pinSize => widget.size ?? _defaultPinSize;
+
+  /// The translation applied to the pin to center its bottom point on the correct position.
+  Offset get _translationOffset => Offset(-pinSize / 2, _bottomOffset - pinSize);
+
+  /// Animation controller to raise the pin when dragged.
+  late final AnimationController _raiseController = AnimationController(
+    vsync: this,
+    duration: const Duration(milliseconds: 600),
+    reverseDuration: const Duration(milliseconds: 120),
+  );
+
+  late final Animation<Offset> _raiseAnimation = Tween<Offset>(
+    begin: Offset.zero,
+    end: const Offset(0, -_heldHeightOffset),
+  ).animate(CurvedAnimation(
+    parent: _raiseController,
+    curve: Curves.elasticOut,
+    reverseCurve: Curves.easeOutQuad,
+  ));
+
+  /// Animation controller to fade the shadow in and out.
+  late final AnimationController _shadowController = AnimationController(
+    vsync: this,
+    duration: const Duration(milliseconds: 200),
+    reverseDuration: const Duration(milliseconds: 150),
+  );
+
+  late final Animation<double> _shadowAnimation = Tween<double>(
+    begin: 0.0,
+    end: 1.0,
+  ).animate(CurvedAnimation(
+    parent: _shadowController,
+    curve: Curves.easeInOut,
+    reverseCurve: Curves.easeInOut,
+  ));
+
+  /// Key associated with the feedback (i.e. cursor/finger-follower) widget shown while dragging the pin.
+  final GlobalKey _feedbackKey = GlobalKey();
+
+  @override
+  void dispose() {
+    _raiseController.dispose();
+    _shadowController.dispose();
+    super.dispose();
+  }
+
+  @override
+  void didUpdateWidget(_StageMapPin oldWidget) {
+    super.didUpdateWidget(oldWidget);
+
+    if (widget.bShowAsDragged != oldWidget.bShowAsDragged) {
+      if (widget.bShowAsDragged) {
+        _playDragStartAnimation();
+      } else {
+        _playDragEndAnimation();
+      }
+    }
+  }
+
+  @override
+  Widget build(BuildContext context) {
+    final selectedActorsPref = Provider.of<SelectedActorSettings>(context, listen: false).selectedActors;
+
+    final Widget stationaryPin = _StageMapPinIcon(
+      selectedActorsPref: selectedActorsPref,
+      actor: widget.actor,
+      size: pinSize,
+      raiseAnimation: _raiseAnimation,
+      shadowAnimation: _shadowAnimation,
+    );
+
+    return Transform.translate(
+      offset: _translationOffset,
+      child: GestureDetector(
+        onTap: widget.bIsSelectable ? _onTapped : null,
+        child: widget.bIsDraggable
+            ? LongPressDraggable(
+                dragAnchorStrategy: childDragAnchorStrategy,
+                ignoringFeedbackSemantics: false,
+                delay: _dragDelay,
+                onDragStarted: _onDragStarted,
+                onDragEnd: _onDragEnd,
+                onDragUpdate: _onDragUpdate,
+
+                // Pin icon while not being dragged
+                child: stationaryPin,
+
+                // Show an identical pin at the drag position while dragged
+                feedback: _StageMapPinIcon(
+                  selectedActorsPref: selectedActorsPref,
+                  actor: widget.actor,
+                  key: _feedbackKey,
+                  size: pinSize,
+                  raiseAnimation: _raiseAnimation,
+                  shadowAnimation: _shadowAnimation,
+                ),
+
+                // Hide the original pin position while dragged
+                childWhenDragging: SizedBox(width: pinSize, height: pinSize),
+              )
+            : stationaryPin,
+      ),
+    );
+  }
+
+  void _onTapped() {
+    final selectedActorSettings = Provider.of<SelectedActorSettings>(context, listen: false);
+    selectedActorSettings.selectActor(
+      widget.actor.path,
+      bShouldSelect: !selectedActorSettings.isActorSelected(widget.actor.path),
+    );
+  }
+
+  void _onDragStarted() {
+    if (widget.onDragStarted != null) {
+      widget.onDragStarted!();
+    }
+
+    _playDragStartAnimation();
+  }
+
+  void _onDragEnd(DraggableDetails details) {
+    if (widget.onDragEnded == null) {
+      return;
+    }
+
+    final Offset? position = _getPinPosition();
+    if (position != null) {
+      widget.onDragEnded!(position);
+    }
+
+    _playDragEndAnimation();
+  }
+
+  void _onDragUpdate(DragUpdateDetails details) {
+    if (widget.onDragUpdate == null) {
+      return;
+    }
+
+    final Offset? position = _getPinPosition();
+    if (position != null) {
+      widget.onDragUpdate!(position);
+    }
+  }
+
+  void _playDragStartAnimation() {
+    _raiseController.forward(from: 0);
+    _shadowController.forward(from: 0);
+  }
+
+  void _playDragEndAnimation() {
+    _raiseController.reverse(from: 1);
+    _shadowController.reverse(from: 1);
+  }
+
+  // Get the current position of the pin (if possible, otherwise null)
+  Offset? _getPinPosition() {
+    final RenderObject? renderObject = _feedbackKey.currentContext?.findRenderObject();
+    if (renderObject == null) {
+      return null;
+    }
+
+    return MatrixUtils.transformPoint(renderObject.getTransformTo(null), -_translationOffset);
+  }
+}
+
+/// The visual component of a draggable pin on the map.
+class _StageMapPinIcon extends StatelessWidget {
+  const _StageMapPinIcon({
+    Key? key,
+    required this.actor,
+    required this.selectedActorsPref,
+    required this.size,
+    required this.raiseAnimation,
+    required this.shadowAnimation,
+  }) : super(key: key);
+
+  /// The stage actor this represents.
+  final UnrealObject actor;
+
+  /// Preference that determines which actors are selected.
+  final TransientPreference<Set<String>> selectedActorsPref;
+
+  /// Base size of the pin.
+  final double size;
+
+  /// The animation controlling the pin's position.
+  final Animation<Offset> raiseAnimation;
+
+  /// The animation controlling the pin's shadow.
+  final Animation<double> shadowAnimation;
+
+  /// The size of the shadow.
+  static const Offset shadowSize = Offset(20, 10);
+
+  @override
+  Widget build(BuildContext context) {
+    const double baseIconSize = 24;
+    const double basePinSize = 96; // Size at which icon is at [baseIconSize]
+    final double pinScale = size / basePinSize;
+
+    final double iconSize = baseIconSize * pinScale;
+    final String? iconPath = actor.getIconPath();
+
+    return Stack(
+      clipBehavior: Clip.none,
+      alignment: Alignment.center,
+      children: [
+        // Shadow
+        Positioned(
+          bottom: -shadowSize.dy / 2,
+          child: AnimatedBuilder(
+            animation: shadowAnimation,
+            builder: (_, __) {
+              return Container(
+                width: shadowSize.dx,
+                height: shadowSize.dy,
+                decoration: BoxDecoration(
+                  borderRadius: BorderRadius.circular(40),
+                  boxShadow: [
+                    BoxShadow(
+                      color: Colors.black.withOpacity(shadowAnimation.value * 0.6),
+                      blurRadius: 5,
+                    ),
+                  ],
+                ),
+              );
+            },
+          ),
+        ),
+
+        // Pin icon
+        SlideTransition(
+          position: raiseAnimation,
+          child: TransientPreferenceBuilder(
+            preference: selectedActorsPref,
+            builder: (context, final Set<String> selectedActors) {
+              final bool bIsSelected = selectedActors.contains(actor.path);
+              final Color mainColor = actor.getUIColor() ?? UnrealColors.highlightBlue;
+
+              return Stack(children: [
+                if (bIsSelected)
+                  AssetIcon(
+                    path: 'assets/images/map/pin_outline.svg',
+                    size: size,
+                    color: const Color(0x99000000),
+                  ),
+                AssetIcon(
+                  path: bIsSelected
+                      ? 'assets/images/map/pin_inner_border.svg'
+                      : 'assets/images/map/pin_inner_border_thin.svg',
+                  size: size,
+                  color: bIsSelected ? const Color(0xffffffff) : mainColor,
+                ),
+                AssetIcon(
+                  path: 'assets/images/map/pin_fill.svg',
+                  size: size,
+                  color: bIsSelected ? mainColor : UnrealColors.gray14,
+                ),
+                if (iconPath != null)
+                  Positioned(
+                    top: 22 * pinScale,
+                    left: (size - iconSize) / 2,
+                    child: AssetIcon(
+                      path: iconPath,
+                      size: iconSize,
+                    ),
+                  ),
+              ]);
+            },
+          ),
+        ),
+      ],
+    );
+  }
+}
+
+/// Unique widget that lets the user pick between control modes.
+class _StageMapControlModeToggle extends StatelessWidget {
+  const _StageMapControlModeToggle({Key? key}) : super(key: key);
+
+  @override
+  Widget build(BuildContext context) {
+    final stageMapSettings = Provider.of<StageMapSettings>(context, listen: false);
+
+    return Container(
+      decoration: BoxDecoration(
+        color: Theme.of(context).colorScheme.surface,
+        borderRadius: BorderRadius.circular(16),
+      ),
+      clipBehavior: Clip.antiAlias,
+      child: PreferenceBuilder(
+        preference: stageMapSettings.bIsInTransformMode,
+        builder: (_, final bool bIsInTransformMode) => Row(children: [
+          _StageMapControlModeToggleButton(
+            iconPath: 'assets/images/icons/world_flat.svg',
+            tooltipMessage: AppLocalizations.of(context)!.stageMapControlModeMap,
+            bIsActive: !bIsInTransformMode,
+            onPressed: () => stageMapSettings.bIsInTransformMode.setValue(false),
+          ),
+          _StageMapControlModeToggleButton(
+            iconPath: 'assets/images/icons/object.svg',
+            tooltipMessage: AppLocalizations.of(context)!.stageMapControlModeObject,
+            bIsActive: bIsInTransformMode,
+            onPressed: () => stageMapSettings.bIsInTransformMode.setValue(true),
+          ),
+        ]),
+      ),
+    );
+  }
+}
+
+/// Buttons within the [_StageMapControlModeToggle].
+class _StageMapControlModeToggleButton extends StatelessWidget {
+  const _StageMapControlModeToggleButton({
+    Key? key,
+    required this.iconPath,
+    required this.tooltipMessage,
+    required this.bIsActive,
+    required this.onPressed,
+  }) : super(key: key);
+
+  /// Path of the icon image file to display on the button.
+  final String iconPath;
+
+  /// Tooltip message to display when this is hovered/long-pressed.
+  final String tooltipMessage;
+
+  /// Whether the mode this represents is active.
+  final bool bIsActive;
+
+  /// Function to call when this is pressed.
+  final void Function() onPressed;
+
+  @override
+  Widget build(BuildContext context) {
+    return EpicIconButton(
+      tooltipMessage: tooltipMessage,
+      onPressed: onPressed,
+      bIsToggledOn: bIsActive,
+      iconPath: iconPath,
+      buttonSize: const Size(48, 36),
+    );
+  }
+}
diff --git a/Engine/Extras/VirtualProduction/EpicStageApp/lib/widgets/elements/list_menu.dart b/Engine/Extras/VirtualProduction/EpicStageApp/lib/widgets/elements/list_menu.dart
new file mode 100644
--- /dev/null
+++ b/Engine/Extras/VirtualProduction/EpicStageApp/lib/widgets/elements/list_menu.dart
@@ -0,0 +1,250 @@
+// Copyright Epic Games, Inc. All Rights Reserved.
+
+import 'package:flutter/material.dart';
+
+import '../../utilities/unreal_colors.dart';
+import 'asset_icon.dart';
+import 'epic_scroll_view.dart';
+import 'modal.dart';
+
+/// Default padding for a list menu item.
+const EdgeInsetsGeometry defaultListMenuItemPadding = EdgeInsets.fromLTRB(18, 2, 4, 2);
+
+/// Which edge of the list menu should be a flattened edge to indicate its origin.
+enum ListMenuEdgeDirection {
+  none,
+  topLeft,
+  topRight,
+  bottomLeft,
+  bottomRight,
+  top,
+  bottom,
+}
+
+/// A menu containing a vertical list of items.
+class ListMenu extends StatelessWidget {
+  const ListMenu({
+    Key? widgetKey,
+    required this.children,
+    this.edgeDirection = ListMenuEdgeDirection.none,
+    this.minWidth = 0,
+    this.maxWidth = double.infinity,
+    this.minHeight = 0,
+    this.backgroundColor,
+    this.bIsScrollable = false,
+  }) : super(key: widgetKey);
+
+  /// Widgets to show in the list.
+  final List<Widget> children;
+
+  /// Which edge of the menu to flatten. Regardless of value, the menu will be the same size.
+  final ListMenuEdgeDirection edgeDirection;
+
+  /// The minimum width of the menu.
+  final double minWidth;
+
+  /// The maximum width of the menu.
+  final double maxWidth;
+
+  /// The minimum height of the menu.
+  final double minHeight;
+
+  /// Optional background color for the list.
+  final Color? backgroundColor;
+
+  /// If true, add a scrollbar and make the inner contents scrollable.
+  final bool bIsScrollable;
+
+  /// Radius of the list menu's borders.
+  static const Radius borderRadius = Radius.circular(8);
+
+  @override
+  Widget build(BuildContext context) {
+    final Color color = backgroundColor ?? Theme.of(context).colorScheme.surfaceVariant;
+
+    late final BorderRadius? boxBorderRadius = BorderRadius.only(
+      topLeft: (edgeDirection != ListMenuEdgeDirection.topLeft && edgeDirection != ListMenuEdgeDirection.top)
+          ? borderRadius
+          : Radius.zero,
+      topRight: (edgeDirection != ListMenuEdgeDirection.topRight && edgeDirection != ListMenuEdgeDirection.top)
+          ? borderRadius
+          : Radius.zero,
+      bottomLeft: (edgeDirection != ListMenuEdgeDirection.bottomLeft && edgeDirection != ListMenuEdgeDirection.bottom)
+          ? borderRadius
+          : Radius.zero,
+      bottomRight: (edgeDirection != ListMenuEdgeDirection.bottomRight && edgeDirection != ListMenuEdgeDirection.bottom)
+          ? borderRadius
+          : Radius.zero,
+    );
+
+    Widget contents = IntrinsicWidth(
+      child: Column(
+        mainAxisSize: MainAxisSize.min,
+        crossAxisAlignment: CrossAxisAlignment.stretch,
+        children: children,
+      ),
+    );
+
+    if (bIsScrollable) {
+      contents = EpicScrollView(child: contents);
+    }
+
+    return Container(
+      constraints: BoxConstraints(
+        minWidth: minWidth,
+        maxWidth: maxWidth,
+        minHeight: minHeight,
+      ),
+      decoration: BoxDecoration(
+        color: color,
+        borderRadius: boxBorderRadius,
+        boxShadow: modalBoxShadow,
+      ),
+      clipBehavior: Clip.antiAlias,
+      child: contents,
+    );
+  }
+}
+
+/// An item shown in a [ListMenu] with a standard icon + text layout.
+class ListMenuSimpleItem extends StatelessWidget {
+  const ListMenuSimpleItem({
+    Key? key,
+    required this.title,
+    this.bIsSelected = false,
+    this.bIsEnabled = true,
+    this.bIsChecked = false,
+    this.bShowCheckbox = false,
+    this.bShowArrow = false,
+    this.iconPath,
+    this.onTap,
+  }) : super(key: key);
+
+  /// The text to show on the item.
+  final String title;
+
+  /// Whether the tile is selected.
+  final bool bIsSelected;
+
+  /// Whether the tile is enabled. If not, it will be greyed out and ignore input.
+  final bool bIsEnabled;
+
+  /// Whether to show a checkbox leading the content.
+  final bool bShowCheckbox;
+
+  /// Whether the checkbox should be checked, if shown.
+  final bool bIsChecked;
+
+  /// Whether to show a chevron trailing the content.
+  final bool bShowArrow;
+
+  /// An optional path to an image file to show as an icon before the item's name.
+  final String? iconPath;
+
+  /// An optional function to call when this list item is tapped.
+  final void Function()? onTap;
+
+  @override
+  Widget build(BuildContext context) {
+    const double arrowSize = 24;
+
+    return MouseRegion(
+      cursor: SystemMouseCursors.click,
+      child: Semantics(
+        selected: bIsSelected,
+        button: true,
+        child: GestureDetector(
+          behavior: HitTestBehavior.opaque,
+          onTap: onTap,
+          child: SizedBox(
+            height: 56,
+            child: Stack(
+              alignment: Alignment.centerLeft,
+              children: [
+                Padding(
+                  padding: EdgeInsets.only(right: bShowArrow ? arrowSize : 0),
+                  child: Row(
+                    mainAxisSize: MainAxisSize.min,
+                    crossAxisAlignment: CrossAxisAlignment.center,
+                    children: [
+                      if (bShowCheckbox)
+                        Padding(
+                          padding: EdgeInsets.only(left: 8, right: 4),
+                          child: AssetIcon(
+                            path: bIsChecked
+                                ? 'assets/images/icons/checkbox_checked.svg'
+                                : 'assets/images/icons/checkbox_unchecked.svg',
+                            size: 18,
+                          ),
+                        ),
+                      if (iconPath != null)
+                        Container(
+                          padding: EdgeInsets.all(16),
+                          child: AssetIcon(
+                            path: iconPath!,
+                            size: 24,
+                          ),
+                        ),
+                      Flexible(
+                        child: Padding(
+                          padding: EdgeInsets.only(
+                            left: iconPath == null ? 12 : 4,
+                            right: bShowArrow ? 0 : 12,
+                          ),
+                          child: Text(
+                            title,
+                            style: Theme.of(context).textTheme.headlineSmall!.copyWith(color: UnrealColors.white),
+                            overflow: TextOverflow.ellipsis,
+                          ),
+                        ),
+                      ),
+                    ],
+                  ),
+                ),
+                if (bShowArrow)
+                  Align(
+                    alignment: Alignment.centerRight,
+                    child: AssetIcon(
+                      path: 'assets/images/icons/chevron_right.svg',
+                      size: arrowSize,
+                    ),
+                  ),
+              ],
+            ),
+          ),
+        ),
+      ),
+    );
+  }
+}
+
+/// A title shown in a [ListMenu] to categorize items.
+class ListMenuHeader extends StatelessWidget {
+  const ListMenuHeader(this.text, {Key? key}) : super(key: key);
+
+  final String text;
+
+  @override
+  Widget build(BuildContext context) {
+    final TextStyle textStyle = Theme.of(context).textTheme.labelMedium!.copyWith(
+          fontSize: 14,
+          letterSpacing: 0.25,
+        );
+
+    return Container(
+      height: 38,
+      padding: EdgeInsets.symmetric(horizontal: 16, vertical: 4),
+      child: Column(
+        mainAxisAlignment: MainAxisAlignment.spaceBetween,
+        crossAxisAlignment: CrossAxisAlignment.start,
+        children: [
+          Text(text, style: textStyle),
+          Container(
+            height: 2,
+            color: UnrealColors.white.withOpacity(0.1),
+          )
+        ],
+      ),
+    );
+  }
+}
