Branch snapshot for CL 27138787

diff --git a/Engine/Extras/VirtualProduction/EpicStageApp/lib/widgets/elements/modal.dart b/Engine/Extras/VirtualProduction/EpicStageApp/lib/widgets/elements/modal.dart
new file mode 100644
--- /dev/null
+++ b/Engine/Extras/VirtualProduction/EpicStageApp/lib/widgets/elements/modal.dart
@@ -0,0 +1,569 @@
+// Copyright Epic Games, Inc. All Rights Reserved.
+
+import 'package:flutter/material.dart';
+import 'package:flutter/services.dart';
+import 'package:flutter_gen/gen_l10n/app_localizations.dart';
+
+import '../../models/navigator_keys.dart';
+import 'delta_widget_base.dart';
+import 'epic_icon_button.dart';
+
+/// Box shadow to show behind modals.
+const List<BoxShadow> modalBoxShadow = [
+  BoxShadow(
+    color: Color(0x66000000),
+    offset: Offset(4, 8),
+    blurRadius: 8,
+  )
+];
+
+/// A route that displays a generic modal dialog centered in the screen.
+class GenericModalDialogRoute<T> extends PopupRoute<T> {
+  GenericModalDialogRoute({
+    required this.builder,
+    this.bResizeToAvoidBottomInset = false,
+    this.bIsBarrierDismissible = true,
+  });
+
+  /// A function that will build the widget to display within this route.
+  final Widget Function(BuildContext context) builder;
+
+  /// If true, resize the route to stay within the bottom inset of the screen (e.g. where the OS keyboard opens on
+  /// mobile devices).
+  final bool bResizeToAvoidBottomInset;
+
+  /// If true, the user can dismiss this modal by tapping outside of its body.
+  final bool bIsBarrierDismissible;
+
+  @override
+  Color? get barrierColor => Colors.black38;
+
+  @override
+  bool get barrierDismissible => bIsBarrierDismissible;
+
+  @override
+  String? get barrierLabel => AppLocalizations.of(rootNavigatorKey.currentContext!)!.modalDismissLabel;
+
+  @override
+  Duration get transitionDuration => Duration(milliseconds: 200);
+
+  @override
+  Duration get reverseTransitionDuration => Duration(milliseconds: 100);
+
+  @override
+  Widget buildPage(BuildContext context, Animation<double> animation, Animation<double> secondaryAnimation) {
+    return CustomSingleChildLayout(
+      delegate: _GenericModalDialogRouteLayout<T>(
+        context: context,
+        bResizeToAvoidBottomInset: bResizeToAvoidBottomInset,
+      ),
+      child: Center(
+        child: Builder(builder: builder),
+      ),
+    );
+  }
+
+  /// Show a dialog using a [GenericModalDialogRoute] as the route.
+  /// If [bIsBarrierDismissible] is true, the user can dismiss this modal by tapping outside of its body.
+  /// Returns a future that will return the result when the dialog was popped from the navigator.
+  static Future<T?> showDialog<T>({
+    required BuildContext context,
+    required WidgetBuilder builder,
+    bool bIsBarrierDismissible = true,
+  }) {
+    final route = GenericModalDialogRoute<T>(
+      bResizeToAvoidBottomInset: true,
+      builder: builder,
+      bIsBarrierDismissible: bIsBarrierDismissible,
+    );
+
+    return Navigator.of(context, rootNavigator: true).push<T>(route);
+  }
+}
+
+class _GenericModalDialogRouteLayout<T> extends SingleChildLayoutDelegate {
+  _GenericModalDialogRouteLayout({
+    required this.context,
+    required this.bResizeToAvoidBottomInset,
+  });
+
+  final BuildContext context;
+
+  /// If true, resize the route to stay within the bottom inset of the screen (e.g. where the OS keyboard opens on
+  /// mobile devices).
+  final bool bResizeToAvoidBottomInset;
+
+  @override
+  BoxConstraints getConstraintsForChild(BoxConstraints constraints) {
+    double maxHeight = double.infinity;
+
+    final MediaQueryData mediaQuery = MediaQuery.of(context);
+    if (bResizeToAvoidBottomInset) {
+      maxHeight = mediaQuery.size.height - mediaQuery.viewInsets.bottom;
+    } else {
+      maxHeight = mediaQuery.size.height;
+    }
+
+    return constraints.copyWith(
+      minHeight: 0,
+      maxHeight: maxHeight,
+    );
+  }
+
+  @override
+  bool shouldRelayout(covariant SingleChildLayoutDelegate oldDelegate) {
+    return oldDelegate != this;
+  }
+}
+
+/// A generic card wrapper for a modal dialog.
+class ModalDialogCard extends StatelessWidget {
+  const ModalDialogCard({
+    Key? key,
+    required this.child,
+    this.color,
+    this.shape,
+  }) : super(key: key);
+
+  /// Contents of the card.
+  final Widget child;
+
+  /// Color of the card.
+  final Color? color;
+
+  /// Shape of the card.
+  final ShapeBorder? shape;
+
+  @override
+  Widget build(BuildContext context) {
+    return Container(
+      decoration: BoxDecoration(boxShadow: modalBoxShadow),
+      constraints: BoxConstraints(minWidth: 300),
+      child: Card(
+        shape: shape,
+        color: color ?? Theme.of(context).colorScheme.surfaceTint,
+        child: child,
+      ),
+    );
+  }
+}
+
+/// Standard title formatting for a modal dialog.
+class ModalDialogTitle extends StatelessWidget {
+  const ModalDialogTitle({Key? key, required this.title}) : super(key: key);
+
+  final String title;
+
+  @override
+  Widget build(BuildContext context) {
+    return ModalDialogSection(
+      child: Padding(
+        padding: const EdgeInsets.symmetric(vertical: 4),
+        child: Text(
+          style: Theme.of(context).textTheme.displayLarge,
+          textAlign: TextAlign.center,
+          this.title,
+        ),
+      ),
+    );
+  }
+}
+
+/// Wrapper to apply standard padding for a section of a modal dialog.
+class ModalDialogSection extends StatelessWidget {
+  const ModalDialogSection({Key? key, required this.child}) : super(key: key);
+
+  final Widget child;
+
+  @override
+  Widget build(BuildContext context) {
+    return Padding(
+      padding: const EdgeInsets.symmetric(
+        horizontal: 16,
+        vertical: 8,
+      ),
+      child: child,
+    );
+  }
+}
+
+/// A generic modal containing information and an "OK" button.
+class InfoModalDialog extends StatelessWidget {
+  const InfoModalDialog({Key? key, required this.message}) : super(key: key);
+
+  /// The localized message to show.
+  final String message;
+
+  /// Show an info dialog containing a [message].
+  static void show(String message) {
+    final BuildContext context = rootNavigatorKey.currentContext!;
+
+    final route = GenericModalDialogRoute(
+      bResizeToAvoidBottomInset: true,
+      builder: (_) => InfoModalDialog(message: message),
+    );
+
+    Navigator.of(context, rootNavigator: true).push(route);
+  }
+
+  /// Show an info dialog containing a message created using a [getMessage] function which will be passed the context in
+  /// which the modal is to be created.
+  static void showInContext(String Function(BuildContext) getMessage) {
+    final BuildContext context = rootNavigatorKey.currentContext!;
+    InfoModalDialog.show(getMessage(context));
+  }
+
+  @override
+  Widget build(BuildContext context) {
+    final localizations = AppLocalizations.of(context)!;
+
+    return ModalDialogCard(
+      child: IntrinsicWidth(
+        child: Column(
+          mainAxisSize: MainAxisSize.min,
+          children: [
+            SizedBox(height: 8),
+            ModalDialogSection(
+              child: Text(message),
+            ),
+            ModalDialogSection(
+              child: Row(
+                mainAxisAlignment: MainAxisAlignment.end,
+                children: [
+                  EpicLozengeButton(
+                    label: localizations.menuButtonOK,
+                    width: 110,
+                    onPressed: () => Navigator.of(context).pop(),
+                  ),
+                ],
+              ),
+            ),
+          ],
+        ),
+      ),
+    );
+  }
+}
+
+/// The action the user chose in a [TextInputModalDialog].
+enum TextInputModalDialogAction {
+  /// The user cancelled the modal.
+  cancel,
+
+  /// The user entered a new value and applied it.
+  apply,
+
+  /// The user pressed the reset button.
+  reset,
+}
+
+/// The result returned when a [TextInputModalDialog] is popped.
+class TextInputModalDialogResult<T> {
+  const TextInputModalDialogResult({required this.action, this.value = null});
+
+  /// The result of the user's interaction with the dialog.
+  final TextInputModalDialogAction action;
+
+  /// The text that the user entered. Will be null unless result is [TextInputModalDialogAction.apply], in which case
+  /// it contains the value parsed from the user's text.
+  final T? value;
+}
+
+/// A generic modal that lets the user input a value using the keyboard.
+/// When popped, this pushes a [TextInputModalDialogResult] with the result of the modal interaction.
+class TextInputModalDialog<T> extends StatefulWidget {
+  const TextInputModalDialog({
+    Key? key,
+    required this.title,
+    required this.parseValue,
+    this.handleResult,
+    this.initialText,
+    this.hintText = '',
+    this.keyboardType,
+    this.inputFormatters,
+    this.bShowResetButton = false,
+  }) : super(key: key);
+
+  /// The localized title of the dialog.
+  final String title;
+
+  /// A function that takes the user's text value and converts it to the desired result type.
+  final T? Function(String) parseValue;
+
+  /// An optional function that takes the result and returns a future which will complete once the modal should close.
+  final Future Function(TextInputModalDialogResult<T>)? handleResult;
+
+  /// The initial text to show in the input field.
+  final String? initialText;
+
+  /// Hint text to display in the text box.
+  final String hintText;
+
+  /// Input type for the text box.
+  final TextInputType? keyboardType;
+
+  /// List of text input formatters to apply to the text.
+  final List<TextInputFormatter>? inputFormatters;
+
+  /// If true, show a button allowing the user to reset the value.
+  final bool bShowResetButton;
+
+  @override
+  State<StatefulWidget> createState() => _TextInputModalDialogState<T>();
+}
+
+class _TextInputModalDialogState<T> extends State<TextInputModalDialog<T>> {
+  /// Text controller for the input field displayed to the user.
+  final _textController = TextEditingController();
+
+  /// If true, this has completed and is waiting for the result to be handled.
+  bool bIsResultBeingHandled = false;
+
+  @override
+  void initState() {
+    super.initState();
+
+    if (widget.initialText != null) {
+      _textController.text = widget.initialText!;
+      _textController.selection = TextSelection(baseOffset: 0, extentOffset: _textController.text.length);
+    }
+  }
+
+  @override
+  Widget build(BuildContext context) {
+    const double buttonWidth = 110;
+    final localizations = AppLocalizations.of(context)!;
+
+    final List<TextInputFormatter> formatters = [];
+    if (widget.inputFormatters != null) {
+      formatters.addAll(widget.inputFormatters!);
+    }
+
+    formatters.add(TextInputFormatter.withFunction(_takeNewIfParsed));
+
+    return Form(
+      autovalidateMode: AutovalidateMode.onUserInteraction,
+      child: ModalDialogCard(
+        child: IntrinsicHeight(
+          child: IntrinsicWidth(
+            child: Column(
+              children: [
+                ModalDialogTitle(title: widget.title),
+                if (bIsResultBeingHandled)
+                  ModalDialogSection(
+                    child: Padding(
+                      padding: EdgeInsets.only(bottom: 16),
+                      child: SizedBox.square(
+                        dimension: 100,
+                        child: CircularProgressIndicator(
+                          strokeWidth: 6,
+                        ),
+                      ),
+                    ),
+                  ),
+                if (!bIsResultBeingHandled)
+                  ModalDialogSection(
+                    child: Row(
+                      children: [
+                        Expanded(
+                          child: SizedBox(
+                            height: 36,
+                            child: TextField(
+                              autofocus: true,
+                              maxLines: 1,
+                              cursorWidth: 1,
+                              controller: _textController,
+                              keyboardAppearance: Brightness.dark,
+                              style: Theme.of(context).textTheme.bodyMedium,
+                              decoration: InputDecoration(hintText: widget.hintText),
+                              keyboardType: widget.keyboardType,
+                              inputFormatters: formatters,
+                              onEditingComplete: _complete,
+                            ),
+                          ),
+                        ),
+                        if (widget.bShowResetButton)
+                          Padding(
+                            padding: const EdgeInsets.only(left: 8),
+                            child: ResetValueButton(onPressed: _reset),
+                          ),
+                      ],
+                    ),
+                  ),
+                if (!bIsResultBeingHandled)
+                  ModalDialogSection(
+                    child: Row(
+                      mainAxisAlignment: MainAxisAlignment.spaceAround,
+                      children: [
+                        EpicLozengeButton(
+                          label: localizations.menuButtonCancel,
+                          width: buttonWidth,
+                          color: Colors.transparent,
+                          onPressed: _cancel,
+                        ),
+                        EpicLozengeButton(
+                          label: localizations.menuButtonOK,
+                          width: buttonWidth,
+                          onPressed: _complete,
+                        ),
+                      ],
+                    ),
+                  ),
+              ],
+            ),
+          ),
+        ),
+      ),
+    );
+  }
+
+  /// Take the new text editing value only if it passes the parsing function.
+  TextEditingValue _takeNewIfParsed(TextEditingValue oldValue, TextEditingValue newValue) {
+    if (widget.parseValue(newValue.text) == null) {
+      return oldValue;
+    }
+
+    return newValue;
+  }
+
+  /// Complete the interaction and pop the modal, returning the entered value.
+  void _complete() async {
+    _handleResult(TextInputModalDialogResult<T>(
+      action: TextInputModalDialogAction.apply,
+      value: widget.parseValue(_textController.value.text),
+    ));
+  }
+
+  /// Cancel the interaction and pop the modal, returning a cancelled result.
+  void _cancel() {
+    _handleResult(TextInputModalDialogResult<T>(action: TextInputModalDialogAction.cancel));
+  }
+
+  /// Cancel the interaction and pop the modal, returning a reset result.
+  void _reset() {
+    _handleResult(TextInputModalDialogResult<T>(action: TextInputModalDialogAction.reset));
+  }
+
+  /// Handle the result of the dialog, either popping or waiting for the provided result handler to complete.
+  void _handleResult(TextInputModalDialogResult<T> result) async {
+    if (bIsResultBeingHandled) {
+      return;
+    }
+
+    if (widget.handleResult != null) {
+      setState(() {
+        bIsResultBeingHandled = true;
+      });
+
+      await widget.handleResult!(result);
+    }
+
+    Navigator.of(context).pop(result);
+  }
+}
+
+/// A generic modal that lets the user input a [double] using the keyboard.
+class DoubleTextInputModalDialog extends StatelessWidget {
+  const DoubleTextInputModalDialog({
+    Key? key,
+    required this.title,
+    this.handleResult,
+    this.initialValue,
+    this.keyboardType,
+    this.inputFormatters,
+    this.bShowResetButton = false,
+  }) : super(key: key);
+
+  /// The localized title of the dialog.
+  final String title;
+
+  /// An optional function that takes the result and returns a future which will complete once the modal should close.
+  final Future Function(TextInputModalDialogResult<double>)? handleResult;
+
+  /// The initial value to show in the input field.
+  final double? initialValue;
+
+  /// Input type for the text box.
+  final TextInputType? keyboardType;
+
+  /// List of text input formatters to apply to the text.
+  final List<TextInputFormatter>? inputFormatters;
+
+  /// If true, show a button allowing the user to reset the value.
+  final bool bShowResetButton;
+
+  @override
+  Widget build(BuildContext context) {
+    return TextInputModalDialog<double>(
+      title: title,
+      handleResult: handleResult,
+      initialText: initialValue?.toStringAsFixed(6),
+      hintText: '0',
+      parseValue: _parseDouble,
+      keyboardType: keyboardType ?? TextInputType.numberWithOptions(decimal: true),
+      bShowResetButton: bShowResetButton,
+      inputFormatters: inputFormatters,
+    );
+  }
+
+  /// Try to convert the text value to a double.
+  double? _parseDouble(String text) {
+    double? result = double.tryParse(text);
+
+    if (result == null) {
+      // Try prepending a 0 so that an empty string or fractional shorthand work
+      result = double.tryParse('0$text');
+    }
+
+    return result;
+  }
+}
+
+/// A generic modal that lets the user input a [String] using the keyboard.
+class StringTextInputModalDialog extends StatelessWidget {
+  const StringTextInputModalDialog({
+    Key? key,
+    required this.title,
+    this.handleResult,
+    this.initialValue,
+    this.hintText = '',
+    this.keyboardType,
+    this.inputFormatters,
+    this.bShowResetButton = false,
+  }) : super(key: key);
+
+  /// The localized title of the dialog.
+  final String title;
+
+  /// An optional function that takes the result and returns a future which will complete once the modal should close.
+  final Future Function(TextInputModalDialogResult<String>)? handleResult;
+
+  /// The initial value to show in the input field.
+  final String? initialValue;
+
+  /// Hint text to display in the text box.
+  final String hintText;
+
+  /// Input type for the text box.
+  final TextInputType? keyboardType;
+
+  /// List of text input formatters to apply to the text.
+  final List<TextInputFormatter>? inputFormatters;
+
+  /// If true, show a button allowing the user to reset the value.
+  final bool bShowResetButton;
+
+  @override
+  Widget build(BuildContext context) {
+    return TextInputModalDialog<String>(
+      title: title,
+      handleResult: handleResult,
+      initialText: initialValue,
+      hintText: hintText,
+      parseValue: (value) => value,
+      keyboardType: keyboardType ?? TextInputType.text,
+      bShowResetButton: bShowResetButton,
+      inputFormatters: inputFormatters,
+    );
+  }
+}
diff --git a/Engine/Extras/VirtualProduction/EpicStageApp/lib/widgets/elements/ndisplay_config_name_display.dart b/Engine/Extras/VirtualProduction/EpicStageApp/lib/widgets/elements/ndisplay_config_name_display.dart
new file mode 100644
--- /dev/null
+++ b/Engine/Extras/VirtualProduction/EpicStageApp/lib/widgets/elements/ndisplay_config_name_display.dart
@@ -0,0 +1,72 @@
+// Copyright Epic Games, Inc. All Rights Reserved.
+
+import 'package:flutter/material.dart';
+import 'package:provider/provider.dart';
+
+import '../../models/settings/selected_actor_settings.dart';
+import '../../models/unreal_actor_manager.dart';
+import '../../models/unreal_types.dart';
+import '../../utilities/constants.dart';
+import 'asset_icon.dart';
+
+class NDisplayConfigNameDisplay extends StatefulWidget {
+  const NDisplayConfigNameDisplay({Key? key}) : super(key: key);
+
+  @override
+  State<NDisplayConfigNameDisplay> createState() => _NDisplayConfigNameDisplayState();
+}
+
+class _NDisplayConfigNameDisplayState extends State<NDisplayConfigNameDisplay> {
+  late final UnrealActorManager _actorManager;
+
+  @override
+  void initState() {
+    super.initState();
+
+    _actorManager = Provider.of<UnrealActorManager>(context, listen: false);
+    _actorManager.watchClassName(nDisplayRootActorClassName, _onRootActorsUpdated);
+  }
+
+  @override
+  void dispose() {
+    _actorManager.stopWatchingClassName(nDisplayRootActorClassName, _onRootActorsUpdated);
+
+    super.dispose();
+  }
+
+  @override
+  Widget build(BuildContext context) {
+    return StreamBuilder(
+      stream: Provider.of<SelectedActorSettings>(context, listen: false).displayClusterRootPath,
+      builder: (context, final AsyncSnapshot<String> displayClusterRootPath) {
+        final UnrealObject? rootActor = _actorManager.getActorAtPath(displayClusterRootPath.data ?? '');
+
+        return Row(
+          crossAxisAlignment: CrossAxisAlignment.center,
+          children: [
+            AssetIcon(
+              path: 'assets/images/icons/ndisplay.svg',
+              size: 20,
+            ),
+            const SizedBox(width: 8),
+            ConstrainedBox(
+              constraints: const BoxConstraints(maxWidth: 200),
+              child: Text(
+                rootActor?.name ?? '???',
+                overflow: TextOverflow.ellipsis,
+                softWrap: false,
+              ),
+            )
+          ],
+        );
+      },
+    );
+  }
+
+  void _onRootActorsUpdated(ActorUpdateDetails details) {
+    if (details.renamedActors.isNotEmpty || details.addedActors.isNotEmpty) {
+      // Force redraw in case we need to update the name/just got the name for an actor we were awaiting
+      setState(() {});
+    }
+  }
+}
diff --git a/Engine/Extras/VirtualProduction/EpicStageApp/lib/widgets/elements/parsed_rich_text.dart b/Engine/Extras/VirtualProduction/EpicStageApp/lib/widgets/elements/parsed_rich_text.dart
new file mode 100644
--- /dev/null
+++ b/Engine/Extras/VirtualProduction/EpicStageApp/lib/widgets/elements/parsed_rich_text.dart
@@ -0,0 +1,92 @@
+// Copyright Epic Games, Inc. All Rights Reserved.
+
+import 'dart:ui';
+
+import 'package:flutter/material.dart';
+import 'package:logging/logging.dart';
+import 'package:xml/xml.dart';
+
+final _log = Logger('ParsedRichText');
+
+/// A function that takes the [parentStyle] of the parent text node and returns a new style for the current node.
+typedef _TextStyleBuilder = TextStyle Function(TextStyle parentStyle);
+
+/// Map from XML node names to builders for the associated text style.
+final Map<String, _TextStyleBuilder> _styleBuildersByTagName = {
+  'b': (parentStyle) => parentStyle.copyWith(
+        fontVariations: [FontVariation('wght', 700)],
+      ),
+  'i': (parentStyle) => parentStyle.copyWith(
+        fontStyle: FontStyle.italic,
+      ),
+};
+
+/// XML node names for which missing style warnings have already been printed.
+final Set<String> _alreadyWarnedMissingStyleNames = {};
+
+/// Parses an input string and converts it to rich text based on XML-style tags.
+class ParsedRichText extends StatelessWidget {
+  const ParsedRichText(
+    this.text, {
+    Key? key,
+    this.style,
+  }) : super(key: key);
+
+  /// The text to parse.
+  final String text;
+
+  /// The base style for the text. If null, use the default style for the context.
+  final TextStyle? style;
+
+  @override
+  Widget build(BuildContext context) {
+    final TextStyle baseStyle = style ?? DefaultTextStyle.of(context).style;
+
+    TextSpan textSpan;
+    try {
+      final XmlNode rootNode = XmlDocumentFragment.parse(text);
+      textSpan = _makeTextSpanFromNode(xmlNode: rootNode, style: baseStyle);
+    } catch (e) {
+      _log.severe('Failed to parse rich text. Text will not be displayed. Input string:\n"$text"\nError: $e');
+      textSpan = TextSpan();
+    }
+
+    return RichText(text: textSpan);
+  }
+
+  /// Given a [xmlNode], return a text span with the appropriate style and containing spans for any of its descendants.
+  /// The provided [style] will be used if the node has no associated style.
+  TextSpan _makeTextSpanFromNode({
+    required XmlNode xmlNode,
+    required TextStyle style,
+  }) {
+    if (xmlNode is XmlElement) {
+      final String nodeName = xmlNode.name.toString();
+      final _TextStyleBuilder? styleBuilder = _styleBuildersByTagName[nodeName];
+
+      if (styleBuilder == null) {
+        if (_alreadyWarnedMissingStyleNames.add(nodeName)) {
+          _log.warning('No style for "$nodeName" tag; falling back to default');
+        }
+      } else {
+        style = styleBuilder.call(style);
+      }
+    }
+
+    if (xmlNode.children.length > 0) {
+      final childrenSpans = xmlNode.children
+          .map((child) => _makeTextSpanFromNode(
+                xmlNode: child,
+                style: style,
+              ))
+          .toList(growable: false);
+
+      return TextSpan(
+        children: childrenSpans,
+        style: style,
+      );
+    }
+
+    return TextSpan(text: xmlNode.text, style: style);
+  }
+}
diff --git a/Engine/Extras/VirtualProduction/EpicStageApp/lib/widgets/elements/place_actor_menu.dart b/Engine/Extras/VirtualProduction/EpicStageApp/lib/widgets/elements/place_actor_menu.dart
new file mode 100644
--- /dev/null
+++ b/Engine/Extras/VirtualProduction/EpicStageApp/lib/widgets/elements/place_actor_menu.dart
@@ -0,0 +1,287 @@
+// Copyright Epic Games, Inc. All Rights Reserved.
+
+import 'package:flutter/material.dart';
+import 'package:flutter_gen/gen_l10n/app_localizations.dart';
+import 'package:provider/provider.dart';
+import 'package:streaming_shared_preferences/streaming_shared_preferences.dart';
+
+import '../../models/settings/recent_actor_settings.dart';
+import '../../models/unreal_actor_creator.dart';
+import '../../models/unreal_types.dart';
+import '../../utilities/unreal_colors.dart';
+import 'dropdown_button.dart';
+import 'dropdown_list_menu.dart';
+import 'epic_icon_button.dart';
+import 'list_menu.dart';
+import 'template_menu.dart';
+
+/// Button that opens the actor placement menu.
+class PlaceActorButton extends StatelessWidget {
+  const PlaceActorButton({Key? key}) : super(key: key);
+
+  @override
+  Widget build(BuildContext context) {
+    return SizedBox(
+      // Extend the button to the full toolbar height to make it easier to tap even though the styling of the button
+      // makes it appear smaller
+      height: double.infinity,
+      child: ModalDropdownButton(
+        buttonBuilder: (_, __) => _PlaceActorButtonContents(),
+        menuBuilder: (_, __) => PlaceActorDropDownMenu(),
+      ),
+    );
+  }
+}
+
+/// The visible contents of a [PlaceActorButton].
+class _PlaceActorButtonContents extends StatelessWidget {
+  const _PlaceActorButtonContents({
+    Key? key,
+    this.bIsActive = false,
+    this.bRoundTop = false,
+    this.bRoundBottom = false,
+  }) : super(key: key);
+
+  /// Whether the button should be styled as an active dropdown.
+  final bool bIsActive;
+
+  /// If true, round the top corners of the container.
+  final bool bRoundTop;
+
+  /// If true, round the bottom corners of the container.
+  final bool bRoundBottom;
+
+  @override
+  Widget build(BuildContext context) {
+    return Container(
+      height: 48,
+      padding: EdgeInsets.symmetric(horizontal: 8),
+      decoration: BoxDecoration(
+        color: bIsActive ? Theme.of(context).colorScheme.surfaceTint : null,
+        borderRadius: BorderRadius.vertical(
+          top: bRoundTop ? Radius.circular(8) : Radius.zero,
+          bottom: bRoundBottom ? Radius.circular(8) : Radius.zero,
+        ),
+      ),
+      child: Align(
+        widthFactor: 1,
+        child: EpicWideButton(
+          text: AppLocalizations.of(context)!.addActorButton,
+          iconPath: 'assets/images/icons/plus.svg',
+          color: bIsActive ? UnrealColors.highlightGreen : Theme.of(context).colorScheme.surfaceTint,
+          iconColor: bIsActive ? null : UnrealColors.highlightGreen,
+          border: Border.all(
+            color: Theme.of(context).colorScheme.surfaceVariant,
+            width: 1,
+          ),
+        ),
+      ),
+    );
+  }
+}
+
+/// Drop-down menu shown to select an actor to place.
+class PlaceActorDropDownMenu extends StatelessWidget {
+  const PlaceActorDropDownMenu({
+    Key? key,
+    this.actorMapPosition,
+    this.bIsFromLongPress = false,
+  }) : super(key: key);
+
+  /// If provided, override the actor's spawn position (in coordinates normalized to the preview map's size).
+  final Offset? actorMapPosition;
+
+  /// Whether the menu was created in-place from a user's long press action.
+  final bool bIsFromLongPress;
+
+  @override
+  Widget build(BuildContext context) {
+    final actorCreator = Provider.of<UnrealActorCreator>(context, listen: false);
+    final actorCreationSettings = ActorCreationSettings(mapPosition: actorMapPosition);
+
+    return DropDownListMenu(
+      bCenterOriginTabOnPivot: bIsFromLongPress,
+      minHeight: 250,
+      originTabBuilder: (context, bIsOnTop) => _PlaceActorButtonContents(
+        bIsActive: true,
+        bRoundTop: bIsFromLongPress && bIsOnTop,
+        bRoundBottom: bIsFromLongPress && !bIsOnTop,
+      ),
+      children: [
+        ListMenuSimpleItem(
+          title: AppLocalizations.of(context)!.actorNameLightCard,
+          iconPath: 'assets/images/icons/light_card.svg',
+          onTap: () {
+            actorCreator.createLightcard(actorCreationSettings);
+            _closeMenu(context);
+          },
+        ),
+        ListMenuSimpleItem(
+          title: AppLocalizations.of(context)!.actorNameFlag,
+          iconPath: 'assets/images/icons/light_card_flag.svg',
+          onTap: () {
+            actorCreator.createFlag(actorCreationSettings);
+            _closeMenu(context);
+          },
+        ),
+        ListMenuSimpleItem(
+          title: AppLocalizations.of(context)!.actorNameChromakeyCard,
+          iconPath: 'assets/images/icons/chromakey_card.svg',
+          onTap: () {
+            actorCreator.createChromakeyCard(actorCreationSettings);
+            _closeMenu(context);
+          },
+        ),
+        ListMenuSimpleItem(
+          title: AppLocalizations.of(context)!.actorNameColorCorrectWindow,
+          iconPath: 'assets/images/icons/color_correct_window.svg',
+          onTap: () {
+            actorCreator.createColorCorrectionWindow(actorCreationSettings);
+            _closeMenu(context);
+          },
+        ),
+        ListMenuSimpleItem(
+          title: AppLocalizations.of(context)!.placeActorMenuAllTemplates,
+          iconPath: 'assets/images/icons/template.svg',
+          bShowArrow: true,
+          onTap: () {
+            _closeMenu(context);
+            TemplatePickerMenu.show(context, actorScreenPosition: actorMapPosition);
+          },
+        ),
+        PreferenceBuilder(
+          preference: Provider.of<RecentActorSettings>(context, listen: false).recentlyPlacedActors,
+          builder: (BuildContext context, List<RecentlyPlacedActorData>? recentlyPlacedActors) {
+            if (recentlyPlacedActors == null) {
+              return SizedBox();
+            }
+
+            return Column(
+              crossAxisAlignment: CrossAxisAlignment.start,
+              children: [
+                if (recentlyPlacedActors.isNotEmpty)
+                  ListMenuHeader(AppLocalizations.of(context)!.placeActorMenuRecentlyPlaced),
+                for (final RecentlyPlacedActorData actorData in recentlyPlacedActors)
+                  _RecentlyPlacedMenuItem(actorData, actorMapPosition: actorMapPosition),
+              ],
+            );
+          },
+        ),
+      ],
+    );
+  }
+
+  /// Close the menu.
+  void _closeMenu(BuildContext context) {
+    Navigator.of(context).pop();
+  }
+}
+
+/// An item in the menu representing a recently placed actor type.
+class _RecentlyPlacedMenuItem extends StatelessWidget {
+  const _RecentlyPlacedMenuItem(this.actorData, {Key? key, this.actorMapPosition}) : super(key: key);
+
+  /// Data about the actor to place.
+  final RecentlyPlacedActorData actorData;
+
+  /// If provided, override the actor's spawn position (in coordinates normalized to the preview map's size).
+  final Offset? actorMapPosition;
+
+  /// The path of the icon to display for this item.
+  String? get _iconPath {
+    if (actorData.templatePath != null) {
+      return 'assets/images/icons/template.svg';
+    }
+
+    switch (actorData.type) {
+      case RecentlyPlacedActorType.lightCard:
+        return 'assets/images/icons/light_card.svg';
+
+      case RecentlyPlacedActorType.flag:
+        return 'assets/images/icons/light_card_flag.svg';
+
+      case RecentlyPlacedActorType.chromakeyCard:
+        return 'assets/images/icons/chromakey_card.svg';
+
+      case RecentlyPlacedActorType.colorCorrectWindow:
+        return 'assets/images/icons/color_correct_window.svg';
+
+      default:
+        return null;
+    }
+  }
+
+  @override
+  Widget build(BuildContext context) {
+    return ListMenuSimpleItem(
+      title: _getActorName(context),
+      iconPath: _iconPath,
+      onTap: () => _createActor(context),
+    );
+  }
+
+  /// The name to display for this view.
+  String _getActorName(BuildContext context) {
+    if (actorData.name != null) {
+      return actorData.name!;
+    }
+
+    switch (actorData.type) {
+      case RecentlyPlacedActorType.lightCard:
+        return AppLocalizations.of(context)!.actorNameLightCard;
+
+      case RecentlyPlacedActorType.flag:
+        return AppLocalizations.of(context)!.actorNameFlag;
+
+      case RecentlyPlacedActorType.chromakeyCard:
+        return AppLocalizations.of(context)!.actorNameChromakeyCard;
+
+      case RecentlyPlacedActorType.colorCorrectWindow:
+        return AppLocalizations.of(context)!.actorNameColorCorrectWindow;
+    }
+  }
+
+  /// Create an actor of the type indicated by this menu item.
+  void _createActor(BuildContext context) {
+    final UnrealActorCreator actorCreator = Provider.of<UnrealActorCreator>(context, listen: false);
+
+    final actorSettings = ActorCreationSettings(mapPosition: actorMapPosition);
+
+    if (actorData.templatePath != null) {
+      actorCreator.createLightcard(actorSettings.copyWith(
+        template: UnrealTemplateData(
+          name: actorData.name ?? actorData.templatePath!,
+          path: actorData.templatePath!,
+        ),
+      ));
+    } else {
+      switch (actorData.type) {
+        case RecentlyPlacedActorType.lightCard:
+          actorCreator.createLightcard(actorSettings);
+          break;
+
+        case RecentlyPlacedActorType.chromakeyCard:
+          actorCreator.createChromakeyCard(actorSettings);
+          break;
+
+        case RecentlyPlacedActorType.flag:
+          actorCreator.createFlag(actorSettings);
+          break;
+
+        case RecentlyPlacedActorType.colorCorrectWindow:
+          actorCreator.createColorCorrectionWindow(actorSettings);
+          break;
+
+        default:
+          throw 'Can\'t create actor of type ${actorData.type}';
+      }
+    }
+
+    _closeMenu(context);
+  }
+
+  /// Close the menu.
+  void _closeMenu(BuildContext context) {
+    Navigator.of(context).pop();
+  }
+}
diff --git a/Engine/Extras/VirtualProduction/EpicStageApp/lib/widgets/elements/reset_mode_button.dart b/Engine/Extras/VirtualProduction/EpicStageApp/lib/widgets/elements/reset_mode_button.dart
new file mode 100644
--- /dev/null
+++ b/Engine/Extras/VirtualProduction/EpicStageApp/lib/widgets/elements/reset_mode_button.dart
@@ -0,0 +1,36 @@
+// Copyright Epic Games, Inc. All Rights Reserved.
+
+import 'package:flutter/material.dart';
+import 'package:flutter_gen/gen_l10n/app_localizations.dart';
+import 'package:provider/provider.dart';
+
+import '../../models/settings/delta_widget_settings.dart';
+import '../../utilities/constants.dart';
+import '../../utilities/transient_preference.dart';
+import 'epic_icon_button.dart';
+
+/// Button to toggle reset mode on and off.
+class ResetModeButton extends StatelessWidget {
+  const ResetModeButton({Key? key}) : super(key: key);
+
+  @override
+  Widget build(BuildContext context) {
+    final deltaSettings = Provider.of<DeltaWidgetSettings>(context);
+
+    return Padding(
+      padding: EdgeInsets.only(right: cardMargin),
+      child: TransientPreferenceBuilder(
+        preference: Provider.of<DeltaWidgetSettings>(context, listen: false).bIsInResetMode,
+        builder: (context, final bool bIsInResetMode) => EpicIconButton(
+          iconPath: 'assets/images/icons/reset_small.svg',
+          iconSize: 17,
+          backgroundSize: const Size(33, 33),
+          tooltipMessage: AppLocalizations.of(context)!.resetModeToggleButtonTooltip,
+          backgroundShape: BoxShape.circle,
+          onPressed: () => deltaSettings.bIsInResetMode.setValue(!bIsInResetMode),
+          bIsToggledOn: bIsInResetMode,
+        ),
+      ),
+    );
+  }
+}
diff --git a/Engine/Extras/VirtualProduction/EpicStageApp/lib/widgets/elements/selector_bar.dart b/Engine/Extras/VirtualProduction/EpicStageApp/lib/widgets/elements/selector_bar.dart
new file mode 100644
--- /dev/null
+++ b/Engine/Extras/VirtualProduction/EpicStageApp/lib/widgets/elements/selector_bar.dart
@@ -0,0 +1,124 @@
+// Copyright Epic Games, Inc. All Rights Reserved.
+
+import 'dart:ui';
+
+import 'package:flutter/material.dart';
+
+import '../../utilities/unreal_colors.dart';
+
+/// A tab-style horizontal bar that allows the user to select a single value from a list.
+class SelectorBar<T> extends StatefulWidget {
+  const SelectorBar({
+    Key? key,
+    required this.value,
+    required this.onSelected,
+    required this.valueNames,
+  }) : super(key: key);
+
+  /// The selected value.
+  final T value;
+
+  /// Function called when a new value is selected.
+  final Function(T newValue) onSelected;
+
+  /// Names to show for each possible value. Values not in this list will not be shown on the widget,
+  /// which can be useful for e.g. a "none" value.
+  final Map<T, String> valueNames;
+
+  @override
+  State<SelectorBar<T>> createState() => _SelectorBarState();
+}
+
+class _SelectorBarState<T> extends State<SelectorBar<T>> with TickerProviderStateMixin {
+  late final TabController _controller;
+
+  List<T> values = [];
+
+  @override
+  void initState() {
+    super.initState();
+
+    _initValues();
+
+    _controller = TabController(
+      initialIndex: values.indexOf(widget.value),
+      length: widget.valueNames.length,
+      vsync: this,
+    );
+  }
+
+  @override
+  void dispose() {
+    _controller.dispose();
+    super.dispose();
+  }
+
+  @override
+  void didUpdateWidget(covariant SelectorBar<T> oldWidget) {
+    super.didUpdateWidget(oldWidget);
+
+    _initValues();
+  }
+
+  @override
+  Widget build(BuildContext context) {
+    final borderRadius = BorderRadius.circular(20);
+    final labelStyle = Theme.of(context).textTheme.bodyMedium!.copyWith(
+          color: UnrealColors.white,
+          fontSize: 12,
+          fontVariations: [FontVariation('wght', 600)],
+          letterSpacing: 0.5,
+        );
+
+    return Container(
+      decoration: BoxDecoration(
+        color: Theme.of(context).colorScheme.background,
+        borderRadius: borderRadius,
+      ),
+      height: 36,
+      child: TabBar(
+        controller: _controller,
+        onTap: (int index) => widget.onSelected(values[index]),
+        labelPadding: EdgeInsets.symmetric(horizontal: 12),
+        isScrollable: true,
+        indicator: BoxDecoration(
+          color: Theme.of(context).colorScheme.secondary,
+          borderRadius: borderRadius,
+        ),
+        labelStyle: labelStyle,
+        unselectedLabelStyle: labelStyle.copyWith(color: UnrealColors.gray75),
+        tabs: [
+          for (final value in values)
+            Tab(
+              child: Text(
+                widget.valueNames[value]!,
+                overflow: TextOverflow.ellipsis,
+              ),
+            )
+        ],
+      ),
+    );
+  }
+
+  /// Set up the list of possible values.
+  void _initValues() {
+    values = widget.valueNames.keys.toList();
+  }
+}
+
+/// "Ornamental" toggle button selector when there's only one value to show.
+class FakeSelectorBar extends StatelessWidget {
+  const FakeSelectorBar(this.valueName, {Key? key}) : super(key: key);
+
+  /// The name to show on the button.
+  final String valueName;
+
+  @override
+  Widget build(BuildContext context) {
+    return SelectorBar(
+      value: true,
+      valueNames: {true: valueName},
+      onSelected: (_) {},
+    );
+  }
+}
diff --git a/Engine/Extras/VirtualProduction/EpicStageApp/lib/widgets/elements/sensitivity_slider.dart b/Engine/Extras/VirtualProduction/EpicStageApp/lib/widgets/elements/sensitivity_slider.dart
new file mode 100644
--- /dev/null
+++ b/Engine/Extras/VirtualProduction/EpicStageApp/lib/widgets/elements/sensitivity_slider.dart
@@ -0,0 +1,117 @@
+// Copyright Epic Games, Inc. All Rights Reserved.
+
+import 'package:flutter/material.dart';
+import 'package:flutter_gen/gen_l10n/app_localizations.dart';
+import 'package:provider/provider.dart';
+import 'package:streaming_shared_preferences/streaming_shared_preferences.dart';
+
+import '../../models/settings/delta_widget_settings.dart';
+import '../../utilities/unreal_colors.dart';
+import 'delta_slider.dart';
+import 'delta_widget_base.dart';
+import 'dropdown_button.dart';
+import 'dropdown_list_menu.dart';
+import 'epic_icon_button.dart';
+
+const double _minSensitivity = 0.2;
+const double _maxSensitivity = 2.0;
+
+/// Convert a sensitivity value (multiplier on base widget speed) to a display value (range of 1-10 for readability).
+double _sensitivityToDisplayValue(double sensitivity) {
+  return sensitivity / _minSensitivity;
+}
+
+/// Inverse of [_sensitivityToDisplayValue].
+double _displayValueToSensitivity(double displayValue) {
+  return displayValue * _minSensitivity;
+}
+
+/// Widget to control the sensitivity of all other delta sliders.
+class DeltaWidgetSensitivitySlider extends StatelessWidget {
+  const DeltaWidgetSensitivitySlider({Key? key}) : super(key: key);
+
+  @override
+  Widget build(BuildContext context) {
+    final deltaSettings = Provider.of<DeltaWidgetSettings>(context, listen: false);
+
+    return PreferenceBuilder(
+      preference: deltaSettings.sensitivity,
+      builder: (context, final double sensitivity) => DrivenDeltaSlider(
+        min: _sensitivityToDisplayValue(_minSensitivity),
+        max: _sensitivityToDisplayValue(_maxSensitivity),
+        bIgnoreSensitivity: true,
+        baseSensitivity: 1.0,
+        label: AppLocalizations.of(context)!.sliderSpeed,
+        values: [DeltaWidgetValueData(value: _sensitivityToDisplayValue(sensitivity))],
+        onChanged: (List<double> values) => deltaSettings.sensitivity.setValue(
+          (sensitivity + _displayValueToSensitivity(values.first)).clamp(_minSensitivity, _maxSensitivity),
+        ),
+        onReset: () => deltaSettings.sensitivity.setValue(1.0),
+      ),
+    );
+  }
+}
+
+/// Button that opens the sensitivity slider menu.
+class SensitivitySliderButton extends StatefulWidget {
+  const SensitivitySliderButton({Key? key}) : super(key: key);
+
+  @override
+  State<SensitivitySliderButton> createState() => _SensitivitySliderButtonState();
+}
+
+class _SensitivitySliderButtonState extends State<SensitivitySliderButton> {
+  @override
+  Widget build(BuildContext context) {
+    final labelStyle = Theme.of(context).textTheme.bodyMedium!.copyWith(
+          color: UnrealColors.white,
+        );
+
+    return ModalDropdownButton(
+      tooltipMessage: AppLocalizations.of(context)!.sliderSpeed,
+      buttonBuilder: (context, state) => Material(
+        color: Colors.transparent,
+        child: Stack(
+          alignment: Alignment.center,
+          children: [
+            EpicIconButton(
+              iconPath: 'assets/images/icons/slider_speed.svg',
+              bIsToggledOn: state != ModalDropdownButtonState.closed,
+              bIsVisualOnly: true,
+            ),
+            // Sensitivity value overlay
+            IgnorePointer(
+              child: SizedBox(
+                width: 24,
+                height: 24,
+                child: Align(
+                  alignment: Alignment.bottomRight,
+                  child: PreferenceBuilder(
+                    preference: Provider.of<DeltaWidgetSettings>(context, listen: false).sensitivity,
+                    builder: (context, double sensitivity) => Text(
+                      _sensitivityToDisplayValue(sensitivity).toStringAsFixed(0),
+                      style: labelStyle,
+                    ),
+                  ),
+                ),
+              ),
+            ),
+          ],
+        ),
+      ),
+      menuBuilder: (context, originTabBuilder) => DropDownListMenu(originTabBuilder: originTabBuilder, children: [
+        Material(
+          color: Colors.transparent,
+          child: Container(
+            width: 300,
+            padding: EdgeInsets.only(
+              top: 8,
+              bottom: 4,
+            ),
+            child: DeltaWidgetSensitivitySlider(),
+          ),
+        ),
+      ]),
+    );
+  }
+}
diff --git a/Engine/Extras/VirtualProduction/EpicStageApp/lib/widgets/elements/spinner_overlay.dart b/Engine/Extras/VirtualProduction/EpicStageApp/lib/widgets/elements/spinner_overlay.dart
new file mode 100644
--- /dev/null
+++ b/Engine/Extras/VirtualProduction/EpicStageApp/lib/widgets/elements/spinner_overlay.dart
@@ -0,0 +1,23 @@
+// Copyright Epic Games, Inc. All Rights Reserved.
+
+import 'package:flutter/material.dart';
+
+/// Spinner overlay used to indicate that we're waiting to receive or load something.
+class SpinnerOverlay extends StatelessWidget {
+  const SpinnerOverlay({Key? key}) : super(key: key);
+
+  @override
+  Widget build(BuildContext context) {
+    return Container(
+      color: Colors.black.withOpacity(0.5),
+      child: const Center(
+        child: SizedBox.square(
+          dimension: 80,
+          child: CircularProgressIndicator(
+            strokeWidth: 6,
+          ),
+        ),
+      ),
+    );
+  }
+}
diff --git a/Engine/Extras/VirtualProduction/EpicStageApp/lib/widgets/elements/stepper.dart b/Engine/Extras/VirtualProduction/EpicStageApp/lib/widgets/elements/stepper.dart
new file mode 100644
--- /dev/null
+++ b/Engine/Extras/VirtualProduction/EpicStageApp/lib/widgets/elements/stepper.dart
@@ -0,0 +1,190 @@
+// Copyright Epic Games, Inc. All Rights Reserved.
+
+import 'dart:ui';
+
+import 'package:flutter/material.dart';
+import 'package:flutter_gen/gen_l10n/app_localizations.dart';
+import 'package:provider/provider.dart';
+
+import '../../models/settings/delta_widget_settings.dart';
+import '../../utilities/transient_preference.dart';
+import '../../utilities/unreal_colors.dart';
+import 'delta_widget_base.dart';
+import 'unreal_widget_base.dart';
+
+/// Configuration for a step button on an [UnrealStepper] widget.
+class StepperStepConfig {
+  const StepperStepConfig({
+    required this.stepSize,
+    required this.label,
+    this.color,
+    this.textColor,
+  });
+
+  /// Delta value to apply for this step.
+  final double stepSize;
+
+  /// Label to show on the step button.
+  final String label;
+
+  /// Color of the button's background.
+  final Color? color;
+
+  /// Color of the text in the button.
+  final Color? textColor;
+
+  /// Step configs for an exposure stepper.
+  static List<StepperStepConfig> exposureSteps = [
+    StepperStepConfig(stepSize: -1, label: '-1', color: Color(0xff1f1f1f)),
+    StepperStepConfig(stepSize: -0.5, label: '-1/2', color: Color(0xff363636)),
+    StepperStepConfig(stepSize: -0.25, label: '-1/4', color: Color(0xff4a4a4a)),
+    StepperStepConfig(stepSize: 0.25, label: '+1/4', color: Color(0xff7a7a7a), textColor: UnrealColors.gray06),
+    StepperStepConfig(stepSize: 0.5, label: '+1/2', color: Color(0xff979797), textColor: UnrealColors.gray06),
+    StepperStepConfig(stepSize: 1, label: '+1', color: Color(0xffc0c0c0), textColor: UnrealColors.gray06),
+  ];
+}
+
+/// A stepper that remotely increments/decrements a property in Unreal Engine.
+class UnrealStepper extends UnrealWidget {
+  const UnrealStepper({
+    Key? key,
+    super.overrideName,
+    super.minMaxBehaviour,
+    super.enableProperties,
+    required super.unrealProperties,
+    required this.steps,
+  }) : super(key: key);
+
+  /// Configurations for each button to display in this widget.
+  final List<StepperStepConfig> steps;
+
+  @override
+  State<StatefulWidget> createState() => _UnrealStepperState();
+}
+
+class _UnrealStepperState extends State<UnrealStepper> with UnrealWidgetStateMixin<UnrealStepper, double> {
+  @override
+  Widget build(BuildContext context) {
+    final List<Widget> children = [];
+
+    for (int stepIndex = 0; stepIndex < widget.steps.length; ++stepIndex) {
+      final bool bIsLast = stepIndex == widget.steps.length - 1;
+
+      final StepperStepConfig config = widget.steps[stepIndex];
+      children.add(Expanded(
+        child: TransientPreferenceBuilder(
+          preference: Provider.of<DeltaWidgetSettings>(context, listen: false).bIsInResetMode,
+          builder: (context, final bool bIsInResetMode) => _StepButton(
+            config: config,
+            onPressed: _onStepButtonPressed,
+            propertyLabel: propertyLabel,
+            bIsEnabled: !bIsInResetMode,
+            bIsFirst: stepIndex == 0,
+            bIsLast: bIsLast,
+          ),
+        ),
+      ));
+    }
+
+    return Padding(
+      padding: EdgeInsets.symmetric(horizontal: DeltaWidgetConstants.widgetOuterXPadding),
+      child: Container(
+        decoration: BoxDecoration(
+          color: Theme.of(context).colorScheme.background,
+          borderRadius: BorderRadius.circular(30),
+        ),
+        padding: EdgeInsets.symmetric(horizontal: 4),
+        child: Row(
+          mainAxisAlignment: MainAxisAlignment.spaceEvenly,
+          children: children,
+        ),
+      ),
+    );
+  }
+
+  /// Called when a step button is pressed.
+  void _onStepButtonPressed(double stepSize) {
+    if (beginTransaction(AppLocalizations.of(context)!.transactionStepProperty(propertyLabel, stepSize))) {
+      final List<double> deltaValues = List.filled(propertyCount, stepSize);
+      handleOnChangedByUser(deltaValues);
+
+      endTransaction();
+    }
+  }
+}
+
+/// A button to apply an increase/decrease in value.
+class _StepButton extends StatelessWidget {
+  const _StepButton({
+    Key? key,
+    required this.config,
+    required this.onPressed,
+    this.propertyLabel,
+    this.bIsEnabled = true,
+    this.bIsFirst = false,
+    this.bIsLast = false,
+  }) : super(key: key);
+
+  /// The configuration for this step button.
+  final StepperStepConfig config;
+
+  /// The label of the property this is modifying.
+  final String? propertyLabel;
+
+  /// Callback for when the button is pressed. Passes the size of the step.
+  final void Function(double stepSize) onPressed;
+
+  /// Whether the owning widget has its controls enabled.
+  final bool bIsEnabled;
+
+  /// Whether this is the first listed button.
+  final bool bIsFirst;
+
+  /// Whether this is the last listed button.
+  final bool bIsLast;
+
+  @override
+  Widget build(BuildContext context) {
+    final TextStyle baseTextStyle = Theme.of(context).textTheme.bodyMedium!;
+    final Color textColor =
+        (config.textColor ?? baseTextStyle.color ?? UnrealColors.gray75).withOpacity(bIsEnabled ? 1 : 0.5);
+
+    final String stepSize = '${config.stepSize > 0 ? '+' : ''}${config.stepSize}';
+
+    return Tooltip(
+      message: AppLocalizations.of(context)!.stepPropertyButtonTooltip(propertyLabel ?? '', stepSize),
+      child: Semantics(
+        button: true,
+        child: MouseRegion(
+          cursor: bIsEnabled ? MaterialStateMouseCursor.clickable : MouseCursor.defer,
+          child: GestureDetector(
+            onTap: bIsEnabled ? () => onPressed(config.stepSize) : null,
+            child: Padding(
+              padding: EdgeInsets.symmetric(vertical: 8, horizontal: 4),
+              child: Container(
+                decoration: BoxDecoration(
+                  color: config.color ?? Theme.of(context).colorScheme.secondary,
+                  borderRadius: BorderRadius.horizontal(
+                    left: bIsFirst ? const Radius.circular(12) : Radius.zero,
+                    right: bIsLast ? const Radius.circular(12) : Radius.zero,
+                  ),
+                ),
+                height: 20,
+                child: Center(
+                  child: Text(
+                    config.label,
+                    style: baseTextStyle.copyWith(
+                      color: textColor,
+                      fontVariations: [FontVariation('wght', 600)],
+                      height: 1.05,
+                    ),
+                  ),
+                ),
+              ),
+            ),
+          ),
+        ),
+      ),
+    );
+  }
+}
diff --git a/Engine/Extras/VirtualProduction/EpicStageApp/lib/widgets/elements/swipe_revealer.dart b/Engine/Extras/VirtualProduction/EpicStageApp/lib/widgets/elements/swipe_revealer.dart
new file mode 100644
--- /dev/null
+++ b/Engine/Extras/VirtualProduction/EpicStageApp/lib/widgets/elements/swipe_revealer.dart
@@ -0,0 +1,433 @@
+// Copyright Epic Games, Inc. All Rights Reserved.
+
+import 'package:flutter/material.dart';
+import 'package:logging/logging.dart';
+
+import '../../utilities/unreal_colors.dart';
+import 'layout/measurable.dart';
+
+final _log = Logger('SwipeRevealer');
+
+/// Function that builds an action widget to reveal on one side of a [SwipeRevealer] when swiped.
+/// Takes a function which, when called, will unlock the widget's swipe state and return it to its resting state.
+/// If [bDestroy] is true, [SwipeRevealer] will transition out, then call a function to indicate that its item has been
+/// deleted.
+typedef SwipeRevealerActionBuilder = Widget Function(
+  BuildContext context,
+  void Function({bool bDeleteItem}) onFinished,
+);
+
+/// A wrapper which allows a widget to be swiped left or right to reveal additional actions.
+class SwipeRevealer extends StatefulWidget {
+  const SwipeRevealer({
+    Key? key,
+    required this.child,
+    this.leftSwipeActionBuilder,
+    this.rightSwipeActionBuilder,
+    this.backgroundPadding = EdgeInsets.zero,
+    this.onDeleted,
+  }) : super(key: key);
+
+  /// The primary widget when this is idle.
+  final Widget child;
+
+  /// Function that builds the widget to reveal when this is swiped left (shown on the right of the widget).
+  final SwipeRevealerActionBuilder? leftSwipeActionBuilder;
+
+  /// Function that builds the widget to reveal when this is swiped right (shown on the left of the widget).
+  final SwipeRevealerActionBuilder? rightSwipeActionBuilder;
+
+  /// Padding to add around the background when any swipe effects are drawn (such as the red background on deletion).
+  final EdgeInsets backgroundPadding;
+
+  /// Function called if this is deleted by one of its swipe actions.
+  final Function()? onDeleted;
+
+  @override
+  State<SwipeRevealer> createState() => _SwipeRevealerState();
+}
+
+class _SwipeRevealerState extends State<SwipeRevealer> with TickerProviderStateMixin {
+  late final AnimationController _backgroundFadeController;
+  late final AnimationController _slideController;
+  late final AnimationController _deleteController;
+
+  late Animation<double> _backgroundFadeAnimation;
+  late Animation<double> _slideAnimation;
+  late Animation<double> _deleteHeightAnimation;
+
+  Size _mainContentSize = Size.zero;
+  Size _leftActionSize = Size.zero;
+  Size _rightActionSize = Size.zero;
+
+  /// How far the user has dragged the slider
+  double _slideExtent = 0;
+
+  /// How much to adjust the slider position to account for action sizes.
+  /// The slider is centered within its stack, so when the actions are of different sizes, it will be off-center.
+  /// This value is calculated to compensate for that effect.
+  double _slideOffset = 0;
+
+  /// If true, the slider is currently locked open on one of the swipe actions.
+  bool _isLockedOpen = false;
+
+  /// Whether this is being deleted due to a swipe action.
+  bool _bIsBeingDeleted = false;
+
+  double get _minSlideExtent => -_leftActionSize.width;
+  double get _maxSlideExtent => _rightActionSize.width;
+
+  @override
+  void initState() {
+    super.initState();
+
+    _slideController = AnimationController(
+      vsync: this,
+      duration: const Duration(milliseconds: 150),
+      lowerBound: 0,
+      upperBound: 1,
+    );
+
+    _backgroundFadeController = AnimationController(
+      vsync: this,
+      duration: const Duration(milliseconds: 150),
+      lowerBound: 0,
+      upperBound: 1,
+    );
+    _backgroundFadeAnimation = _backgroundFadeController.drive(Tween<double>(
+      begin: 0,
+      end: 1,
+    ));
+
+    _deleteController = AnimationController(
+      duration: const Duration(milliseconds: 400),
+      vsync: this,
+    );
+
+    _updateSlideAnimation();
+  }
+
+  @override
+  void dispose() {
+    _slideController.dispose();
+    _backgroundFadeController.dispose();
+    _deleteController.dispose();
+
+    super.dispose();
+  }
+
+  @override
+  Widget build(BuildContext context) {
+    return LayoutBuilder(builder: (context, final BoxConstraints constraints) {
+      Widget mainContents = Measurable(
+        onMeasured: _onMainContentMeasured,
+        child: widget.child,
+      );
+
+      if (_mainContentSize == Size.zero) {
+        return mainContents;
+      }
+
+      // Ignore pointer events on the child when locked open. This will instead be handled by returning to resting
+      // state instead.
+      if (_bIsBeingDeleted || _isLockedOpen) {
+        mainContents = IgnorePointer(
+          child: mainContents,
+        );
+      }
+
+      // Touch inputs for swipe + unlock tap
+      if (!_bIsBeingDeleted) {
+        mainContents = GestureDetector(
+          behavior: HitTestBehavior.opaque,
+          onTap: _isLockedOpen ? _returnToResting : null,
+          onHorizontalDragStart: _onDragStart,
+          onHorizontalDragUpdate: _onDragUpdate,
+          onHorizontalDragEnd: _onDragEnd,
+          child: mainContents,
+        );
+      }
+
+      final Widget? leftAction = _buildSwipeActionWidget(
+        bIsLeftSide: true,
+        builder: widget.leftSwipeActionBuilder,
+        measuredCallback: _onLeftActionMeasured,
+        currentSize: _leftActionSize,
+      );
+
+      final Widget? rightAction = _buildSwipeActionWidget(
+        bIsLeftSide: false,
+        builder: widget.rightSwipeActionBuilder,
+        measuredCallback: _onRightActionMeasured,
+        currentSize: _rightActionSize,
+      );
+
+      final mainStack = Stack(
+        key: const Key('mainStack'),
+        clipBehavior: Clip.none,
+        children: [
+          // Background highlight shown when being dragged
+          Positioned(
+            left: widget.backgroundPadding.left,
+            right: widget.backgroundPadding.right,
+            top: widget.backgroundPadding.top,
+            bottom: widget.backgroundPadding.bottom,
+            child: FadeTransition(
+              opacity: _backgroundFadeAnimation,
+              child: Container(
+                color: UnrealColors.gray13,
+              ),
+            ),
+          ),
+
+          // Red background shown when an item is being deleted
+          if (_bIsBeingDeleted)
+            AnimatedBuilder(
+              animation: _slideAnimation,
+              child: Container(
+                color: UnrealColors.highlightRed,
+              ),
+              builder: (context, child) => Positioned(
+                left: widget.backgroundPadding.left + _mainContentSize.width + _slideAnimation.value,
+                right: widget.backgroundPadding.right,
+                top: widget.backgroundPadding.top,
+                bottom: widget.backgroundPadding.bottom,
+                child: child!,
+              ),
+            ),
+
+          // Row containing revealed swipe actions and main content
+          ClipRect(
+            clipBehavior: Clip.hardEdge,
+            child: SizedBox(
+              height: _mainContentSize.height,
+              child: OverflowBox(
+                maxWidth: double.infinity,
+                maxHeight: double.infinity,
+                child: AnimatedBuilder(
+                  animation: _slideAnimation,
+
+                  child: Row(children: [
+                    if (rightAction != null) rightAction,
+                    ConstrainedBox(
+                      constraints: BoxConstraints(maxWidth: constraints.maxWidth),
+                      child: mainContents,
+                    ),
+                    if (leftAction != null) leftAction,
+                  ]),
+
+                  // Position the row based on animation
+                  builder: (context, child) => Transform(
+                    transform: Matrix4.translationValues(
+                      _slideAnimation.value + _slideOffset,
+                      0,
+                      0,
+                    ),
+                    transformHitTests: true,
+                    child: child!,
+                  ),
+                ),
+              ),
+            ),
+          ),
+        ],
+      );
+
+      // Ignore input while deleting
+      if (_bIsBeingDeleted) {
+        return IgnorePointer(
+          child: AnimatedBuilder(
+            animation: _deleteHeightAnimation,
+            child: mainStack,
+            builder: (context, child) => SizedBox(
+              height: _deleteHeightAnimation.value,
+              child: child!,
+            ),
+          ),
+        );
+      }
+
+      return TapRegion(
+        behavior: HitTestBehavior.translucent,
+        onTapOutside: (_) => _returnToResting(),
+        child: mainStack,
+      );
+    });
+  }
+
+  /// Make a swipe action widget wrapped with [Measurable] and [Offstage] widgets as necessary.
+  /// Returns null if no builder was provided.
+  Widget? _buildSwipeActionWidget({
+    required bool bIsLeftSide,
+    required SwipeRevealerActionBuilder? builder,
+    required Function(Size) measuredCallback,
+    required Size currentSize,
+  }) {
+    if (builder == null) {
+      return null;
+    }
+
+    Widget swipeActionWidget = Measurable(
+      onMeasured: measuredCallback,
+      child: builder(
+        context,
+        ({bool bDeleteItem = false}) => _onSwipeActionFinished(
+          bIsLeftSide: bIsLeftSide,
+          bDeleteItem: bDeleteItem,
+        ),
+      ),
+    );
+
+    // If size hasn't been determined, draw offstage to avoid flickering
+    if (currentSize.width == 0) {
+      swipeActionWidget = Offstage(child: swipeActionWidget);
+    }
+
+    // Disable input if this is about to be deleted
+    if (_bIsBeingDeleted) {
+      swipeActionWidget = IgnorePointer(child: swipeActionWidget);
+    }
+
+    return swipeActionWidget;
+  }
+
+  /// Called when a swipe action is finished.
+  void _onSwipeActionFinished({required bool bIsLeftSide, bool bDeleteItem = false}) {
+    if (bDeleteItem) {
+      _delete();
+    } else {
+      _returnToResting();
+    }
+  }
+
+  /// Play the delete animation.
+  void _delete() {
+    setState(() {
+      _bIsBeingDeleted = true;
+
+      _slideAnimation = Tween<double>(
+        begin: _slideAnimation.value,
+        end: -_mainContentSize.width - _leftActionSize.width,
+      ).animate(CurvedAnimation(
+        parent: _deleteController,
+        curve: const Interval(0.0, 0.7, curve: Curves.easeInOut),
+      ));
+
+      _deleteHeightAnimation = Tween<double>(
+        begin: _mainContentSize.height,
+        end: 0,
+      ).animate(CurvedAnimation(
+        parent: _deleteController,
+        curve: const Interval(0.3, 1.0, curve: Curves.easeInOut),
+      ));
+
+      _deleteController.forward();
+
+      _deleteController.addStatusListener((status) {
+        if (status == AnimationStatus.completed) {
+          if (widget.onDeleted == null) {
+            _log.warning('${this.widget} was deleted, but has no onDeleted function');
+          } else {
+            widget.onDeleted!();
+          }
+        }
+      });
+    });
+  }
+
+  /// Called when this should unlock its swiped state and return to its resting state.
+  void _returnToResting() {
+    setState(() {
+      _isLockedOpen = false;
+    });
+
+    _slideController.forward();
+    _backgroundFadeController.reverse();
+  }
+
+  /// Called when we learn the size of the main child widget's contents.
+  void _onMainContentMeasured(Size size) {
+    if (_bIsBeingDeleted || !mounted) {
+      // Ignore during delete animation since we're going to transition it out regardless
+      return;
+    }
+
+    _mainContentSize = size;
+    _updateSlideAnimation();
+  }
+
+  /// Called when we learn the size of the left action widget.
+  void _onLeftActionMeasured(Size size) {
+    if (!mounted) {
+      return;
+    }
+
+    setState(() {
+      _leftActionSize = size;
+      _updateSlideOffset();
+    });
+  }
+
+  /// Called when we learn the size of the right action widget.
+  void _onRightActionMeasured(Size size) {
+    if (!mounted) {
+      return;
+    }
+
+    setState(() {
+      _rightActionSize = size;
+      _updateSlideOffset();
+    });
+  }
+
+  /// Called when either swipe action changes size and we need to adjust the slide offset accordingly.
+  void _updateSlideOffset() {
+    _slideOffset = (_leftActionSize.width - _rightActionSize.width) / 2;
+  }
+
+  /// Update the slide animation so its start point matches the current slide extent.
+  void _updateSlideAnimation() {
+    // Determine offset relative to size of main content
+    final double xOffset;
+    if (_mainContentSize.width > 0) {
+      xOffset = _slideExtent;
+    } else {
+      xOffset = 0;
+    }
+
+    setState(() {
+      _slideAnimation = _slideController.drive(Tween<double>(
+        begin: xOffset,
+        end: 0,
+      ));
+    });
+  }
+
+  /// Called when the user starts a horizontal drag gesture on this.
+  void _onDragStart(DragStartDetails details) {
+    _slideExtent = _slideAnimation.value;
+    _slideController.reset();
+    _updateSlideAnimation();
+
+    _backgroundFadeController.forward();
+  }
+
+  /// Called when the user continues a horizontal drag gesture on this.
+  void _onDragUpdate(DragUpdateDetails details) {
+    _slideExtent += details.primaryDelta!;
+    _slideExtent = _slideExtent.clamp(_minSlideExtent, _maxSlideExtent);
+    _updateSlideAnimation();
+  }
+
+  /// Called when the user finishes a horizontal drag gesture on this.
+  void _onDragEnd(DragEndDetails details) {
+    // Lock open if at min/max extents
+    if (_slideExtent == _minSlideExtent || _slideExtent == _maxSlideExtent) {
+      setState(() {
+        _isLockedOpen = true;
+      });
+    } else {
+      _returnToResting();
+    }
+  }
+}
diff --git a/Engine/Extras/VirtualProduction/EpicStageApp/lib/widgets/elements/template_menu.dart b/Engine/Extras/VirtualProduction/EpicStageApp/lib/widgets/elements/template_menu.dart
new file mode 100644
--- /dev/null
+++ b/Engine/Extras/VirtualProduction/EpicStageApp/lib/widgets/elements/template_menu.dart
@@ -0,0 +1,254 @@
+// Copyright Epic Games, Inc. All Rights Reserved.
+
+import 'dart:async';
+
+import 'package:flutter/material.dart';
+import 'package:flutter_gen/gen_l10n/app_localizations.dart';
+import 'package:provider/provider.dart';
+
+import '../../models/engine_connection.dart';
+import '../../models/unreal_actor_creator.dart';
+import '../../models/unreal_types.dart';
+import '../search_bar.dart';
+import 'epic_icon_button.dart';
+import 'epic_list_view.dart';
+import 'layout/card.dart';
+import 'list_menu.dart';
+
+/// Drop-down menu shown to select a template to create an actor from.
+class TemplatePickerMenu extends StatefulWidget {
+  const TemplatePickerMenu({Key? key, this.actorMapPosition}) : super(key: key);
+
+  /// If provided, override the actor's spawn position (in coordinates normalized to the preview map's size).
+  final Offset? actorMapPosition;
+
+  /// Show the template selection menu.
+  static void show(
+    BuildContext context, {
+    Offset? actorScreenPosition,
+  }) {
+    showDialog(
+      context: context,
+      builder: (context) => TemplatePickerMenu(actorMapPosition: actorScreenPosition),
+      barrierDismissible: true,
+      barrierColor: Colors.black38,
+    );
+  }
+
+  @override
+  State<TemplatePickerMenu> createState() => _TemplatePickerMenuState();
+}
+
+class _TemplatePickerMenuState extends State<TemplatePickerMenu> {
+  /// How long to wait before showing a loading spinner.
+  static const Duration _showSpinnerDelay = Duration(milliseconds: 500);
+
+  /// Whether to show a spinner to indicate that we're waiting to receive data.
+  bool _bShowSpinner = true;
+
+  /// A timer that will show the loading spinner when it completes.
+  Timer? _spinnerTimer = null;
+
+  /// List of available templates.
+  final List<UnrealTemplateData> _templates = [];
+
+  final TextEditingController _searchTextController = TextEditingController();
+
+  @override
+  void initState() {
+    super.initState();
+
+    _fetchTemplates();
+
+    _searchTextController.addListener(_onSearchTextChanged);
+  }
+
+  /// Called when a menu item with the given [templateIndex] is pressed.
+  void _onItemPressed(int templateIndex) {
+    final UnrealActorCreator actorCreator = Provider.of<UnrealActorCreator>(context, listen: false);
+
+    if (templateIndex >= 0 && templateIndex < _templates.length) {
+      final UnrealTemplateData templateData = _templates[templateIndex];
+      actorCreator.createLightcard(ActorCreationSettings(
+        template: templateData,
+        mapPosition: widget.actorMapPosition,
+      ));
+    }
+
+    _closeMenu();
+  }
+
+  /// Close the menu.
+  void _closeMenu() {
+    Navigator.of(context).pop();
+  }
+
+  /// Fetch the list of available templates from the engine.
+  void _fetchTemplates() async {
+    final EngineConnectionManager connectionManager = Provider.of<EngineConnectionManager>(context, listen: false);
+
+    // Don't immediately show the spinner to avoid flickering, but start a timer to show it
+    _spinnerTimer?.cancel();
+    _spinnerTimer = Timer(_showSpinnerDelay, _showSpinner);
+
+    final UnrealHttpResponse response = await connectionManager.sendHttpRequest(UnrealHttpRequest(
+      url: '/remote/search/assets',
+      verb: 'PUT',
+      body: {
+        'Query': _searchTextController.text,
+        'Filter': {
+          'ClassNames': ['/Script/DisplayClusterLightCardEditor.DisplayClusterLightCardTemplate'],
+          'Limit': 100,
+        },
+      },
+    ));
+
+    // No need for spinner now that we've received data.
+    _spinnerTimer?.cancel();
+
+    _handleTemplatesResponse(response);
+
+    setState(() {
+      _bShowSpinner = false;
+    });
+  }
+
+  /// Handle a response from the engine containing the list of templates.
+  void _handleTemplatesResponse(UnrealHttpResponse response) {
+    _templates.clear();
+
+    if (response.code != HttpResponseCode.ok) {
+      return;
+    }
+
+    final dynamic assets = response.body['Assets'];
+    if (assets == null || assets is! List<dynamic>) {
+      return;
+    }
+
+    for (final dynamic entry in assets) {
+      final String? name = entry['Name'];
+      if (name == null) {
+        continue;
+      }
+
+      final String? path = entry['Path'];
+      if (path == null) {
+        continue;
+      }
+
+      _templates.add(UnrealTemplateData(name: name, path: path));
+    }
+  }
+
+  /// Called when the user's search query text has changed.
+  void _onSearchTextChanged() {
+    _fetchTemplates();
+  }
+
+  @override
+  Widget build(BuildContext context) {
+    final Widget itemsWidget;
+
+    if (_bShowSpinner) {
+      itemsWidget = Center(
+        child: SizedBox.square(
+          dimension: 60,
+          child: CircularProgressIndicator(
+            strokeWidth: 6,
+          ),
+        ),
+      );
+    } else {
+      if (_templates.length == 0) {
+        itemsWidget = Padding(
+          padding: EdgeInsets.only(top: 16),
+          child: Text(
+            AppLocalizations.of(context)!.templateMenuEmptyMessage,
+            textAlign: TextAlign.center,
+          ),
+        );
+      } else {
+        itemsWidget = EpicListView(
+          padding: EdgeInsets.only(top: 4),
+          itemCount: _templates.length,
+          itemBuilder: (context, index) => _TemplateMenuItem(
+            templateData: _templates[index],
+            onPressed: () => _onItemPressed(index),
+          ),
+        );
+      }
+    }
+
+    return MediaQuery.removeViewInsets(
+      context: context,
+      removeLeft: true,
+      removeTop: true,
+      removeRight: true,
+      removeBottom: true,
+      child: Center(
+        child: Card(
+          child: ListMenu(
+            maxWidth: 270,
+            minWidth: 270,
+            children: [
+              CardLargeHeader(
+                iconPath: 'assets/images/icons/viewport.svg',
+                title: AppLocalizations.of(context)!.placeActorMenuAllTemplates,
+                trailing: EpicIconButton(
+                  iconPath: 'assets/images/icons/close.svg',
+                  tooltipMessage: AppLocalizations.of(context)!.menuButtonCancel,
+                  iconSize: 20,
+                  onPressed: _closeMenu,
+                ),
+              ),
+              CardSubHeader(
+                padding: EdgeInsets.symmetric(vertical: 8, horizontal: 16),
+                child: EpicSearchBar(controller: _searchTextController),
+              ),
+              ConstrainedBox(
+                constraints: BoxConstraints(
+                  minHeight: 400,
+                  maxHeight: 400,
+                ),
+                child: itemsWidget,
+              ),
+            ],
+          ),
+        ),
+      ),
+    );
+  }
+
+  /// Enable the spinner if safe.
+  void _showSpinner() {
+    if (!mounted) {
+      return;
+    }
+
+    setState(() {
+      _bShowSpinner = true;
+    });
+  }
+}
+
+/// An item in the menu representing a placeable template.
+class _TemplateMenuItem extends StatelessWidget {
+  const _TemplateMenuItem({
+    Key? key,
+    required this.templateData,
+    required this.onPressed,
+  }) : super(key: key);
+
+  final UnrealTemplateData templateData;
+  final VoidCallback onPressed;
+
+  @override
+  Widget build(BuildContext context) {
+    return ListMenuSimpleItem(
+      title: templateData.name,
+      iconPath: 'assets/images/icons/template.svg',
+      onTap: onPressed,
+    );
+  }
+}
diff --git a/Engine/Extras/VirtualProduction/EpicStageApp/lib/widgets/elements/transform_gesture_detector.dart b/Engine/Extras/VirtualProduction/EpicStageApp/lib/widgets/elements/transform_gesture_detector.dart
new file mode 100644
--- /dev/null
+++ b/Engine/Extras/VirtualProduction/EpicStageApp/lib/widgets/elements/transform_gesture_detector.dart
@@ -0,0 +1,80 @@
+// Copyright Epic Games, Inc. All Rights Reserved.
+
+import 'package:flutter/cupertino.dart';
+import 'package:flutter/gestures.dart';
+
+import '../../gestures/rotate_and_scale.dart';
+
+class TransformGestureDetector extends StatelessWidget {
+  const TransformGestureDetector({
+    Key? key,
+    this.child,
+    this.behavior,
+    this.onScaleStart,
+    this.onScaleUpdate,
+    this.onScaleEnd,
+    this.onPanStart,
+    this.onPanUpdate,
+    this.onPanEnd,
+    this.onLongPressStart,
+  }) : super(key: key);
+
+  final Widget? child;
+
+  /// How this gesture detector should behave during hit testing.
+  final HitTestBehavior? behavior;
+
+  /// The pointers in contact with the screen have established a focal point and
+  /// initial scale of 1.0.
+  final GestureScaleStartCallback? onScaleStart;
+
+  /// The pointers in contact with the screen have indicated a new focal point
+  /// and/or scale.
+  final GestureScaleUpdateCallback? onScaleUpdate;
+
+  /// The pointers are no longer in contact with the screen.
+  final GestureScaleEndCallback? onScaleEnd;
+
+  /// A pointer has contacted the screen with a primary button and has begun to
+  /// move.
+  final GestureDragStartCallback? onPanStart;
+
+  /// A pointer that is in contact with the screen with a primary button and
+  /// moving has moved again.
+  final GestureDragUpdateCallback? onPanUpdate;
+
+  /// A pointer that was previously in contact with the screen with a primary
+  /// button and moving is no longer in contact with the screen and was moving
+  /// at a specific velocity when it stopped contacting the screen.
+  final GestureDragEndCallback? onPanEnd;
+
+  /// A pointer has performed a long press in a single position on the screen.
+  final GestureLongPressStartCallback? onLongPressStart;
+
+  @override
+  Widget build(BuildContext context) {
+    return RawGestureDetector(
+      behavior: behavior,
+      child: child,
+      gestures: {
+        if (onScaleStart != null || onScaleUpdate != null || onScaleEnd != null)
+          RotateAndScaleGestureRecognizer: GestureRecognizerFactoryWithHandlers<RotateAndScaleGestureRecognizer>(
+            () => RotateAndScaleGestureRecognizer(),
+            (RotateAndScaleGestureRecognizer instance) {
+              instance
+                ..onStart = onScaleStart
+                ..onUpdate = onScaleUpdate
+                ..onEnd = onScaleEnd;
+            },
+          ),
+        if (onLongPressStart != null)
+          LongPressGestureRecognizer: GestureRecognizerFactoryWithHandlers<LongPressGestureRecognizer>(
+            () => LongPressGestureRecognizer(),
+            (LongPressGestureRecognizer instance) {
+              instance.onLongPressStart = onLongPressStart;
+            },
+          ),
+      },
+    );
+  }
+}
diff --git a/Engine/Extras/VirtualProduction/EpicStageApp/lib/widgets/elements/tree_view.dart b/Engine/Extras/VirtualProduction/EpicStageApp/lib/widgets/elements/tree_view.dart
new file mode 100644
--- /dev/null
+++ b/Engine/Extras/VirtualProduction/EpicStageApp/lib/widgets/elements/tree_view.dart
@@ -0,0 +1,260 @@
+// Copyright Epic Games, Inc. All Rights Reserved.
+
+import 'dart:collection';
+
+import 'package:flutter/material.dart';
+
+import 'epic_list_view.dart';
+
+/// A list of items shown in a collapsible, nested tree structure.
+class TreeView extends StatefulWidget {
+  const TreeView({
+    Key? key,
+    required this.treeController,
+    required this.nodeBuilder,
+    this.padding = EdgeInsets.zero,
+  }) : super(key: key);
+
+  /// Controller which provides the node data and state management for the tree.
+  final TreeViewController treeController;
+
+  /// Function which will create a node widget for the given node data.
+  final Widget Function(TreeViewNode node, TreeViewController controller) nodeBuilder;
+
+  /// Additional padding around the main scrollable area.
+  final EdgeInsets padding;
+
+  @override
+  State<TreeView> createState() => _TreeViewState();
+}
+
+class _TreeViewState extends State<TreeView> {
+  @override
+  void initState() {
+    super.initState();
+
+    widget.treeController.addListener(_onControllerChanged);
+  }
+
+  @override
+  void didUpdateWidget(covariant TreeView oldWidget) {
+    super.didUpdateWidget(oldWidget);
+
+    if (oldWidget.treeController != widget.treeController) {
+      oldWidget.treeController.removeListener(_onControllerChanged);
+      widget.treeController.addListener(_onControllerChanged);
+      _onControllerChanged();
+    }
+  }
+
+  @override
+  void dispose() {
+    widget.treeController.removeListener(_onControllerChanged);
+
+    super.dispose();
+  }
+
+  @override
+  Widget build(BuildContext context) {
+    return MediaQuery.removePadding(
+      context: context,
+      removeBottom: true,
+      child: EpicListView(
+        itemCount: widget.treeController.getVisibleNodeCount(),
+        itemBuilder: _buildNode,
+        padding: widget.padding,
+      ),
+    );
+  }
+
+  Widget _buildNode(BuildContext context, int index) {
+    return widget.nodeBuilder(
+      widget.treeController.getVisibleNodeWithIndex(index),
+      widget.treeController,
+    );
+  }
+
+  void _onControllerChanged() {
+    setState(() {});
+  }
+}
+
+/// A controller for a [TreeView] which holds data about the tree's nodes and manages their states.
+class TreeViewController extends ChangeNotifier {
+  /// List of nodes at the root level of the tree.
+  final List<TreeViewNode> _children = [];
+
+  /// Map from unique keys to nodes in the tree.
+  final Map<String, TreeViewNode> _nodesByKey = {};
+
+  /// Set of keys for nodes which are currently collapsed.
+  final Set<String> _collapsedNodes = {};
+
+  /// Cached list of currently visible in nodes in the order they appear. This may not contain all currently visible
+  /// nodes.
+  final List<TreeViewNode> _cachedVisibleNodes = [];
+
+  /// Cached number of currently visible nodes.
+  int? _cachedVisibleNodeCount = null;
+
+  /// Get a list of all node keys currently in the tree.
+  Iterable<String> get allKeys => UnmodifiableListView(_nodesByKey.keys);
+
+  /// Add a node to the tree.
+  /// If provided, the node will be added as a child of the node specified by [parentKey].
+  void addNode(TreeViewNode node, {String? parentKey}) {
+    if (_nodesByKey.containsKey(node.key)) {
+      throw ArgumentError('Tried to add a node with key ${node.key}, but that key is already in the tree');
+    }
+
+    TreeViewNode? parent = null;
+    if (parentKey != null) {
+      parent = _nodesByKey[parentKey];
+      if (parent == null) {
+        throw ArgumentError('Tried to add a node with parent key ${node.key}, but that key is not in the tree');
+      }
+    }
+
+    _nodesByKey[node.key] = node;
+
+    if (parent != null) {
+      parent.children.add(node);
+      node.indentation = parent.indentation + 1;
+      node.parent = parent;
+    } else {
+      _children.add(node);
+    }
+
+    _invalidateVisibleNodeCache();
+    notifyListeners();
+  }
+
+  /// Remove a node from the tree if it exists.
+  bool removeNode(String key) {
+    if (!_nodesByKey.containsKey(key)) {
+      return false;
+    }
+
+    final node = _nodesByKey.remove(key);
+    _children.remove(node);
+
+    node!.parent?.children.remove(node);
+
+    _invalidateVisibleNodeCache();
+    notifyListeners();
+    return true;
+  }
+
+  /// Expand the node with the given [key].
+  void expandNode(String key) {
+    if (_collapsedNodes.remove(key)) {
+      _invalidateVisibleNodeCache();
+      notifyListeners();
+    }
+  }
+
+  /// Collapse the node with the given [key].
+  void collapseNode(String key) {
+    if (_collapsedNodes.add(key)) {
+      _invalidateVisibleNodeCache();
+      notifyListeners();
+    }
+  }
+
+  /// Toggle the collapsed/expanded state of the node with the given [key].
+  void toggleNode(String key) {
+    if (isNodeExpanded(key)) {
+      collapseNode(key);
+    } else {
+      expandNode(key);
+    }
+  }
+
+  /// Get the node with the given key if it exists.
+  TreeViewNode? getNode(String key) => _nodesByKey[key];
+
+  /// Get the number of visible nodes.
+  int getVisibleNodeCount() {
+    if (_cachedVisibleNodeCount == null) {
+      _cachedVisibleNodeCount = _children.fold<int>(0, (sum, child) => sum + _getVisibleNodeCount(child));
+    }
+
+    return _cachedVisibleNodeCount!;
+  }
+
+  /// Get the node with the given index in the flattened list of visible nodes.
+  TreeViewNode getVisibleNodeWithIndex(int index) {
+    if (index < 0 || index >= getVisibleNodeCount()) {
+      throw ArgumentError('Visible node index out of range');
+    }
+
+    if (index >= _cachedVisibleNodes.length) {
+      _cacheVisibleNodeIndices(_children, 0);
+    }
+
+    return _cachedVisibleNodes[index];
+  }
+
+  /// Check whether the node with the given [key] is expanded.
+  bool isNodeExpanded(String key) {
+    return !_collapsedNodes.contains(key);
+  }
+
+  /// Invalidate the visible node cache.
+  void _invalidateVisibleNodeCache() {
+    _cachedVisibleNodeCount = null;
+    _cachedVisibleNodes.clear();
+  }
+
+  /// Given a list of [nodes], traverse them and fill [_cachedVisibleNodes] with the nodes at each index.
+  /// Treat the first encountered node as index [index].
+  int _cacheVisibleNodeIndices(List<TreeViewNode> nodes, int index) {
+    assert(index <= _cachedVisibleNodes.length);
+
+    for (final TreeViewNode node in nodes) {
+      if (_cachedVisibleNodes.length <= index) {
+        _cachedVisibleNodes.add(node);
+      }
+      index += 1;
+
+      if (isNodeExpanded(node.key)) {
+        index = _cacheVisibleNodeIndices(node.children, index);
+      }
+    }
+
+    return index;
+  }
+
+  /// Count how many visible nodes are contained within a [node] and its descendants (including itself).
+  int _getVisibleNodeCount(TreeViewNode node) {
+    if (!isNodeExpanded(node.key)) {
+      return 1;
+    }
+
+    return node.children.fold<int>(1, (sum, child) => sum + _getVisibleNodeCount(child));
+  }
+}
+
+/// Data about a node in a [TreeView].
+class TreeViewNode<T> {
+  TreeViewNode({
+    required this.key,
+    required this.data,
+    List<TreeViewNode>? children,
+  }) : this.children = children?.toList() ?? [];
+
+  /// A string used to refer to this node uniquely within the tree.
+  final String key;
+
+  /// List of nodes descending from this one.
+  final List<TreeViewNode> children;
+
+  /// Custom data associated with the node.
+  final T data;
+
+  /// How many layers deep this node is indented.
+  int indentation = 0;
+
+  /// The parent of this node in its current tree structure, or null if it's root-level.
+  TreeViewNode? parent;
+}
diff --git a/Engine/Extras/VirtualProduction/EpicStageApp/lib/widgets/elements/unreal_property_builder.dart b/Engine/Extras/VirtualProduction/EpicStageApp/lib/widgets/elements/unreal_property_builder.dart
new file mode 100644
--- /dev/null
+++ b/Engine/Extras/VirtualProduction/EpicStageApp/lib/widgets/elements/unreal_property_builder.dart
@@ -0,0 +1,273 @@
+// Copyright Epic Games, Inc. All Rights Reserved.
+
+import 'package:flutter/material.dart';
+import 'package:provider/provider.dart';
+
+import '../../models/property_modify_operations.dart';
+import '../../models/unreal_property_manager.dart';
+import '../../models/unreal_types.dart';
+import '../../utilities/guarded_refresh_state.dart';
+
+/// Signature for a function to modify a remote property in Unreal Engine.
+typedef PropertyModifyFunction<T> = void Function(
+  PropertyModifyOperation operation, {
+  required String description,
+  T? deltaValue,
+  PropertyMinMaxBehaviour minMaxBehaviour,
+  T? overrideMin,
+  T? overrideMax,
+});
+
+/// Uses the context's [UnrealPropertyManager] to track a property, then builds a widget exposing its value and a
+/// function to modify it.
+class UnrealPropertyBuilder<T> extends StatefulWidget {
+  const UnrealPropertyBuilder({
+    Key? key,
+    required this.property,
+    required this.builder,
+  }) : super(key: key);
+
+  /// The property to track.
+  final UnrealProperty property;
+
+  /// Function that takes the build [context], the current [value] of the property (or null if not available yet), and a
+  /// [modify] function which can be called to modify the property.
+  final Widget Function(BuildContext context, T? value, PropertyModifyFunction<T> modify) builder;
+
+  @override
+  State<UnrealPropertyBuilder> createState() => _UnrealPropertyBuilderState<T>();
+}
+
+class _UnrealPropertyBuilderState<T> extends State<UnrealPropertyBuilder<T>> with GuardedRefreshState {
+  late final UnrealPropertyManager _propertyManager;
+  TrackedPropertyId? _trackedPropertyId;
+
+  @override
+  void initState() {
+    super.initState();
+
+    _propertyManager = Provider.of<UnrealPropertyManager>(context, listen: false);
+    _trackProperty();
+  }
+
+  @override
+  void dispose() {
+    _stopTrackingProperty();
+
+    super.dispose();
+  }
+
+  @override
+  void didUpdateWidget(covariant UnrealPropertyBuilder<T> oldWidget) {
+    super.didUpdateWidget(oldWidget);
+
+    if (oldWidget.property != widget.property) {
+      _trackProperty();
+    }
+  }
+
+  @override
+  Widget build(BuildContext context) {
+    return widget.builder(
+      context,
+      _propertyManager.getTrackedPropertyValue(_trackedPropertyId!),
+      _modifyProperty,
+    );
+  }
+
+  /// Modify the tracked property.
+  void _modifyProperty(
+    PropertyModifyOperation operation, {
+    required String description,
+    T? deltaValue,
+    PropertyMinMaxBehaviour minMaxBehaviour = PropertyMinMaxBehaviour.clamp,
+    T? overrideMin,
+    T? overrideMax,
+  }) {
+    if (_propertyManager.beginTransaction(description)) {
+      _propertyManager.modifyTrackedPropertyValue(
+        _trackedPropertyId!,
+        operation,
+        deltaValue: deltaValue,
+        minMaxBehaviour: minMaxBehaviour,
+        overrideMin: overrideMin,
+        overrideMax: overrideMax,
+      );
+
+      _propertyManager.endTransaction();
+    }
+  }
+
+  /// Start tracking the widget's property (and stop tracking the old one if one exists).
+  void _trackProperty() {
+    if (_trackedPropertyId != null) {
+      _stopTrackingProperty();
+    }
+
+    _trackedPropertyId = _propertyManager.trackProperty(widget.property, _onPropertyChanged);
+    _propertyManager.waitForProperty(_trackedPropertyId!).then(refreshOnData);
+  }
+
+  /// Stop tracking the currently tracked property.
+  void _stopTrackingProperty() {
+    if (_trackedPropertyId != null) {
+      _propertyManager.stopTrackingProperty(_trackedPropertyId!, _onPropertyChanged);
+    }
+  }
+
+  /// Called when the tracked property changes.
+  void _onPropertyChanged(dynamic newValue) {
+    guardedRefresh();
+  }
+}
+
+/// Uses the context's [UnrealPropertyManager] to track multiple properties of the same type, then builds a widget
+/// exposing the shared value and a function to modify it.
+class UnrealMultiPropertyBuilder<T> extends StatefulWidget {
+  const UnrealMultiPropertyBuilder({
+    Key? key,
+    required this.properties,
+    required this.fallbackValue,
+    required this.builder,
+  }) : super(key: key);
+
+  /// The properties to track.
+  final List<UnrealProperty> properties;
+
+  /// If the properties have different values, this will be passed down as the shared value instead.
+  final T fallbackValue;
+
+  /// Function that takes the build [context], the current [sharedValue] shared by all properties (or null if not
+  /// available yet, or the [fallbackValue] if the values aren't all the same), and a [modify] function which can be
+  /// called to modify the property.
+  final Widget Function(BuildContext context, T? sharedValue, PropertyModifyFunction<T> modify) builder;
+
+  @override
+  State<UnrealMultiPropertyBuilder> createState() => _UnrealMultiPropertyBuilderState<T>();
+}
+
+class _UnrealMultiPropertyBuilderState<T> extends State<UnrealMultiPropertyBuilder<T>> with GuardedRefreshState {
+  late final UnrealPropertyManager _propertyManager;
+  Iterable<TrackedPropertyId> _trackedPropertyIds = [];
+  bool _bIsWaitingForValues = true;
+
+  @override
+  void initState() {
+    super.initState();
+
+    _propertyManager = Provider.of<UnrealPropertyManager>(context, listen: false);
+    _trackProperties();
+  }
+
+  @override
+  void dispose() {
+    _stopTrackingProperties();
+
+    super.dispose();
+  }
+
+  @override
+  void didUpdateWidget(covariant UnrealMultiPropertyBuilder<T> oldWidget) {
+    super.didUpdateWidget(oldWidget);
+
+    if (oldWidget.properties.length != widget.properties.length) {
+      _trackProperties();
+      return;
+    }
+
+    for (int propertyIndex = 0; propertyIndex < oldWidget.properties.length; ++propertyIndex) {
+      if (oldWidget.properties[propertyIndex] != widget.properties[propertyIndex]) {
+        _trackProperties();
+        return;
+      }
+    }
+  }
+
+  @override
+  Widget build(BuildContext context) {
+    return widget.builder(
+      context,
+      _getSharedValue(),
+      _modifyProperty,
+    );
+  }
+
+  /// Modify the tracked property.
+  void _modifyProperty(
+    PropertyModifyOperation operation, {
+    required String description,
+    T? deltaValue,
+    PropertyMinMaxBehaviour minMaxBehaviour = PropertyMinMaxBehaviour.clamp,
+    T? overrideMin,
+    T? overrideMax,
+  }) {
+    if (_propertyManager.beginTransaction(description)) {
+      for (final TrackedPropertyId propertyId in _trackedPropertyIds) {
+        _propertyManager.modifyTrackedPropertyValue(
+          propertyId,
+          operation,
+          deltaValue: deltaValue,
+          minMaxBehaviour: minMaxBehaviour,
+          overrideMin: overrideMin,
+          overrideMax: overrideMax,
+        );
+      }
+
+      _propertyManager.endTransaction();
+    }
+  }
+
+  /// Start tracking the widget's property (and stop tracking the old one if one exists).
+  void _trackProperties() async {
+    if (_trackedPropertyIds.isNotEmpty) {
+      _stopTrackingProperties();
+    }
+
+    _trackedPropertyIds =
+        widget.properties.map((property) => _propertyManager.trackProperty(property, _onPropertyChanged));
+
+    // Refresh once all values are received
+    _bIsWaitingForValues = true;
+    await Future.wait(_trackedPropertyIds.map((propertyId) => _propertyManager.waitForProperty(propertyId)));
+
+    _bIsWaitingForValues = false;
+    guardedRefresh();
+  }
+
+  /// Stop tracking the currently tracked property.
+  void _stopTrackingProperties() {
+    if (_trackedPropertyIds.isEmpty) {
+      return;
+    }
+
+    for (final TrackedPropertyId propertyId in _trackedPropertyIds) {
+      _propertyManager.stopTrackingProperty(propertyId, _onPropertyChanged);
+    }
+  }
+
+  /// Get the value to pass down to the builder function.
+  T? _getSharedValue() {
+    if (_bIsWaitingForValues) {
+      return null;
+    }
+
+    final List<T?> values = _trackedPropertyIds
+        .map((propertyId) => _propertyManager.getTrackedPropertyValue(propertyId) as T?)
+        .toList(growable: false);
+
+    final firstValue = values.first;
+
+    for (int propertyIndex = 1; propertyIndex < values.length; ++propertyIndex) {
+      if (values[propertyIndex] != firstValue) {
+        return widget.fallbackValue;
+      }
+    }
+
+    return firstValue;
+  }
+
+  /// Called when the tracked property changes.
+  void _onPropertyChanged(dynamic newValue) {
+    guardedRefresh();
+  }
+}
