Branch snapshot for CL 27138787

diff --git a/Engine/Extras/VirtualProduction/EpicStageApp/lib/widgets/elements/unreal_widget_base.dart b/Engine/Extras/VirtualProduction/EpicStageApp/lib/widgets/elements/unreal_widget_base.dart
new file mode 100644
--- /dev/null
+++ b/Engine/Extras/VirtualProduction/EpicStageApp/lib/widgets/elements/unreal_widget_base.dart
@@ -0,0 +1,497 @@
+// Copyright Epic Games, Inc. All Rights Reserved.
+
+import 'package:flutter/cupertino.dart';
+import 'package:flutter/foundation.dart';
+import 'package:flutter_gen/gen_l10n/app_localizations.dart';
+import 'package:logging/logging.dart';
+import 'package:provider/provider.dart';
+
+import '../../models/property_modify_operations.dart';
+import '../../models/unreal_property_manager.dart';
+import '../../models/unreal_types.dart';
+import 'delta_widget_base.dart';
+
+final _log = Logger('UnrealWidget');
+
+/// Represents a single value shared by all controlled properties of a widget (if such a value exists).
+class SingleSharedValue<T> {
+  const SingleSharedValue({required this.value, required this.bHasMultipleValues});
+
+  /// The value shared by all controlled properties. This may be null either if all property values are null or if
+  /// there's no single shared value.
+  final T? value;
+
+  /// True if there is no single shared value among all properties.
+  final bool bHasMultipleValues;
+
+  bool operator ==(Object other) {
+    if (identical(this, other)) {
+      return true;
+    }
+
+    if (other.runtimeType != runtimeType) {
+      return false;
+    }
+
+    return other is SingleSharedValue<T> && other.value == value && other.bHasMultipleValues == bHasMultipleValues;
+  }
+}
+
+/// Data for an Unreal property controlled by a widget. Each widget should have its own instance of this class per
+/// value that the widget is controlling.
+class WidgetControlledUnrealProperty<T> {
+  WidgetControlledUnrealProperty(
+      {required this.trackedId,
+      required this.propertyIndex,
+      required this.changedCallback,
+      required this.value,
+      this.ownerName});
+
+  /// The ID used to refer to this property with the property manager.
+  final TrackedPropertyId trackedId;
+
+  /// The index of this property in the list of properties the widget is displaying.
+  final int propertyIndex;
+
+  /// The callback function when the property manager changes the value of the property.
+  final PropertyChangeCallback changedCallback;
+
+  /// The name of the object this property belongs to.
+  final String? ownerName;
+
+  /// The current value of the property, or null if it hasn't been set yet.
+  T? value;
+
+  /// Generate the data needed to display a dot for this property.
+  /// This should only be called once the [value] is ready (i.e. not null).
+  DeltaWidgetValueData<T> toValueData() {
+    if (value == null) {
+      throw Exception('toValueData should not be called until the property\'s value is non-null');
+    }
+
+    return DeltaWidgetValueData(value: value!);
+  }
+}
+
+/// A widget that controls one or more Unreal properties.
+abstract class UnrealWidget extends StatefulWidget {
+  const UnrealWidget({
+    Key? key,
+    required this.unrealProperties,
+    this.overrideName,
+    this.enableProperties,
+    this.onChangedByUser,
+    this.minMaxBehaviour = PropertyMinMaxBehaviour.clamp,
+  }) : super(key: key);
+
+  /// An optional name that will override the property name received from the engine.
+  final String? overrideName;
+
+  /// The properties this widget controls in the engine.
+  final List<UnrealProperty> unrealProperties;
+
+  /// How to treat values that exceed the engine min/max.
+  final PropertyMinMaxBehaviour minMaxBehaviour;
+
+  /// A callback function called when the user changes the value of the widget.
+  final void Function()? onChangedByUser;
+
+  /// A list of properties controlling whether the value is enabled in the engine. If provided, a checkbox will be
+  /// shown controlling these properties.
+  final List<UnrealProperty>? enableProperties;
+}
+
+/// Mixin that adds state management for widgets that control Unreal properties.
+/// To use this mixin:
+/// - Call [makeValueDataList] to create a list of [DeltaWidgetValueData] containing the current values.
+/// - Call [handleOnChangedByUser] with a list of delta values when the user interacts with the widget.
+/// - Call [handleOnResetByUser] when the user presses the widget's reset button.
+/// - Override [modifyOperation] if needed.
+mixin UnrealWidgetStateMixin<WidgetType extends UnrealWidget, PropertyType> on State<WidgetType> {
+  /// The property manager this widget uses to update its properties.
+  late final UnrealPropertyManager _propertyManager;
+
+  /// A list of properties currently controlled by the widget.
+  final List<WidgetControlledUnrealProperty<PropertyType>?> _properties = [];
+
+  /// A list of properties for enabling widget properties current controlled by the widget.
+  final List<WidgetControlledUnrealProperty<bool>?> _enableProperties = [];
+
+  /// The last list of properties that were passed to the widget. This is stored so we can compare the list on rebuild
+  /// and register/unregister properties when the list changes.
+  final List<UnrealProperty> _lastUnrealProperties = [];
+
+  /// The last list of enable properties that were passed to the widget. This is stored so we can compare the list on
+  /// rebuild and register/unregister properties when the list changes.
+  final List<UnrealProperty> _lastEnableProperties = [];
+
+  /// The property's metadata as received from the engine.
+  UnrealPropertyMetadata? _propertyMetadata;
+
+  /// This is incremented whenever the set of tracked properties changes so we can ignore any callbacks from previous
+  /// versions. For example, if the user rapidly changes the list multiple times, we may get changed value or newly
+  /// tracked value callbacks for properties that we no longer care about.
+  int _latestTrackedPropertyVersion = 0;
+
+  /// Whether the properties' values are ready to display to the user.
+  bool _bIsReady = false;
+
+  /// Whether the widget is actively in a transaction.
+  bool _bIsInTransaction = false;
+
+  /// Whether the properties' values are ready to display to the user.
+  bool get bIsReady => _bIsReady;
+
+  /// The property's metadata as received from the engine.
+  UnrealPropertyMetadata? get propertyMetadata => _propertyMetadata;
+
+  /// The label of the property to show to users.
+  String get propertyLabel => widget.overrideName ?? (propertyMetadata?.displayName ?? '');
+
+  /// The property's minimum value as received from the engine.
+  PropertyType? get engineMin => propertyMetadata?.minValue;
+
+  /// The property's maximum value as received from the engine.
+  PropertyType? get engineMax => propertyMetadata?.maxValue;
+
+  /// The overridden minimum value the controlled properties can reach. If null, [engineMin] will be used.
+  PropertyType? get overrideMin => null;
+
+  /// The overridden maximum value the controlled properties can reach. If null, [engineMax] will be used.
+  PropertyType? get overrideMax => null;
+
+  /// The property's possible enum values as received from the engine.
+  List<String> get propertyEnumValues => propertyMetadata?.enumValues ?? [];
+
+  /// The number of properties this is controlling.
+  int get propertyCount => _properties.length;
+
+  /// The type of modify operation this widget uses to apply deltas.
+  PropertyModifyOperation get modifyOperation => const AddOperation();
+
+  /// Extra meta-data properties to pass to conversion functions operating on this widget's value.
+  Map<String, dynamic>? get conversionMetadata => null;
+
+  @override
+  void initState() {
+    super.initState();
+
+    _propertyManager = Provider.of<UnrealPropertyManager>(context, listen: false);
+    _trackAllProperties();
+  }
+
+  @override
+  void dispose() {
+    for (final WidgetControlledUnrealProperty<PropertyType>? property in _properties) {
+      if (property != null) {
+        _propertyManager.stopTrackingProperty(property.trackedId, property.changedCallback);
+      }
+    }
+
+    for (final WidgetControlledUnrealProperty<bool>? property in _enableProperties) {
+      if (property != null) {
+        _propertyManager.stopTrackingProperty(property.trackedId, property.changedCallback);
+      }
+    }
+
+    super.dispose();
+  }
+
+  @override
+  void didUpdateWidget(WidgetType oldWidget) {
+    super.didUpdateWidget(oldWidget);
+
+    _trackAllPropertiesIfChanged();
+  }
+
+  /// Call this to indicate that the user has started interacting with the widget.
+  /// This will be automatically called for you as soon as any properties change, but you can safely call it earlier.
+  /// A [description] will be automatically generated if null is passed.
+  /// Returns true if a new transaction was started.
+  bool beginTransaction([String? description]) {
+    if (_bIsInTransaction) {
+      return false;
+    }
+
+    if (_propertyManager
+        .beginTransaction(description ?? AppLocalizations.of(context)!.transactionEditProperty(propertyLabel))) {
+      _bIsInTransaction = true;
+      return true;
+    }
+
+    return false;
+  }
+
+  /// Call this to indicate that the user has finished interacting with the widget.
+  void endTransaction() {
+    if (!_bIsInTransaction) {
+      return;
+    }
+
+    _bIsInTransaction = false;
+    _propertyManager.endTransaction();
+  }
+
+  /// Generate a list of [DeltaWidgetValueData] containing the data needed to display each control in the widget.
+  /// Values may be null if the value hasn't been received from the engine yet.
+  List<DeltaWidgetValueData<PropertyType>?> makeValueDataList() =>
+      _bIsReady ? _properties.map((property) => property?.toValueData()).toList() : [];
+
+  /// Called when the user changes controlled values.
+  void handleOnChangedByUser(List<PropertyType> deltaValues) {
+    modifyProperties(modifyOperation, values: deltaValues);
+  }
+
+  /// Called when the user presses the reset button.
+  void handleOnResetByUser() {
+    if (!_propertyManager.beginTransaction(AppLocalizations.of(context)!.transactionResetProperty(propertyLabel))) {
+      return;
+    }
+
+    for (final WidgetControlledUnrealProperty<PropertyType>? property in _properties) {
+      if (property != null) {
+        _propertyManager.modifyTrackedPropertyValue(property.trackedId, const ResetOperation());
+      }
+    }
+
+    if (widget.enableProperties != null) {
+      for (final WidgetControlledUnrealProperty<bool>? property in _enableProperties) {
+        if (property != null) {
+          _propertyManager.modifyTrackedPropertyValue(property.trackedId, const ResetOperation());
+        }
+      }
+    }
+
+    _propertyManager.endTransaction();
+  }
+
+  /// Return a SingleSharedValue indicating whether there is a single value shared by all controlled properties.
+  /// If there are no controlled values, the SingleSharedValue will instead contain [defaultValue].
+  SingleSharedValue<PropertyType> getSingleSharedValue([PropertyType? defaultValue]) {
+    // Get the value of all non-null properties
+    final List<PropertyType?> controlledValues = _properties.map((property) => property?.value).toList();
+    controlledValues.removeWhere((element) => element == null);
+
+    bool bHasMultipleValues = false;
+    PropertyType? sharedValue;
+
+    if (controlledValues.length > 1) {
+      sharedValue = controlledValues[0];
+
+      // We're controlling multiple properties. Check that they all have the same value
+      for (int valueIndex = 1; valueIndex < controlledValues.length; ++valueIndex) {
+        if (controlledValues[valueIndex] != sharedValue) {
+          bHasMultipleValues = true;
+          sharedValue = null;
+          break;
+        }
+      }
+    } else if (controlledValues.isNotEmpty) {
+      // There's exactly one value, so use it directly
+      sharedValue = controlledValues[0];
+    } else {
+      // There are no properties, so return the default
+      sharedValue = defaultValue;
+    }
+
+    return SingleSharedValue(value: sharedValue, bHasMultipleValues: bHasMultipleValues);
+  }
+
+  /// Apply an [operation] to each property this controls by the amounts specified in [values].
+  /// If [values] is null, this is treated as if it was a list of null values matching the number of properties.
+  /// If [bIgnoreLimits] is true, the min and max for the widget will be ignored.
+  void modifyProperties(PropertyModifyOperation operation, {List<dynamic>? values, bool bIgnoreLimits = false}) {
+    assert(values == null || values.length == _properties.length);
+
+    if (!_bIsInTransaction) {
+      if (!beginTransaction()) {
+        // If we can't start a transaction, ignore user input or we'll get out of sync with the engine
+        return;
+      }
+    }
+
+    // Enable the property in the engine
+    for (final WidgetControlledUnrealProperty<bool>? property in _enableProperties) {
+      if (property != null) {
+        if (!_propertyManager.getTrackedPropertyValue(property.trackedId)) {
+          _propertyManager.modifyTrackedPropertyValue(property.trackedId, const SetOperation(), deltaValue: true);
+        }
+      }
+    }
+
+    for (int propertyIndex = 0; propertyIndex < _properties.length; ++propertyIndex) {
+      final WidgetControlledUnrealProperty<PropertyType>? property = _properties[propertyIndex];
+
+      if (property != null) {
+        _propertyManager.modifyTrackedPropertyValue(
+          property.trackedId,
+          operation,
+          deltaValue: values?[propertyIndex],
+          minMaxBehaviour: bIgnoreLimits ? PropertyMinMaxBehaviour.ignore : widget.minMaxBehaviour,
+          overrideMin: overrideMin,
+          overrideMax: overrideMax,
+        );
+      }
+    }
+
+    if (widget.onChangedByUser != null) {
+      widget.onChangedByUser!();
+    }
+  }
+
+  /// Check whether the widget's list of properties has changed, and if so, update our list of tracked properties.
+  void _trackAllPropertiesIfChanged() {
+    bool bHavePropertiesChanged = widget.unrealProperties.length != _lastUnrealProperties.length ||
+        (widget.enableProperties?.length ?? 0) != _lastEnableProperties.length;
+
+    if (!bHavePropertiesChanged) {
+      bHavePropertiesChanged = !listEquals(widget.unrealProperties, _lastUnrealProperties);
+    }
+
+    if (!bHavePropertiesChanged && widget.enableProperties != null) {
+      bHavePropertiesChanged = !listEquals(widget.enableProperties!, _lastEnableProperties);
+    }
+
+    if (!bHavePropertiesChanged) {
+      return;
+    }
+
+    _trackAllProperties();
+  }
+
+  /// Resets the internal list of properties and starts tracking them again.
+  void _trackAllProperties() {
+    _bIsReady = false;
+    ++_latestTrackedPropertyVersion;
+
+    // The tracked property version when we started this function, so we can ignore callbacks later in the function
+    // if they come back too late to be relevant.
+    final int trackedPropertyVersion = _latestTrackedPropertyVersion;
+
+    // Futures that will complete when each property is exposed.
+    final List<Future<void>> exposeFutures = [];
+
+    _trackRelevantProperties<PropertyType>(
+      trackedProperties: _properties,
+      oldProperties: _lastUnrealProperties,
+      newProperties: widget.unrealProperties,
+      exposeFutures: exposeFutures,
+      callback: _handleOnManagedValueChanged,
+    );
+
+    _trackRelevantProperties<bool>(
+      trackedProperties: _enableProperties,
+      oldProperties: _lastEnableProperties,
+      newProperties: widget.enableProperties ?? [],
+      exposeFutures: exposeFutures,
+      callback: _handleOnEnableValueChanged,
+    );
+
+    // Wait for all properties to be subscribed, then mark the widget as ready to be used
+    Future.wait(exposeFutures).then((_) {
+      if (_latestTrackedPropertyVersion != trackedPropertyVersion || !mounted) {
+        return;
+      }
+
+      _onPropertiesReady();
+    }).onError((error, stackTrace) {
+      _log.severe('Error while waiting for Unreal widget properties', error, stackTrace);
+    });
+  }
+
+  /// Given a list of [trackedProperties], a corresponding list of [oldProperties] that are being tracked, and a list
+  /// of [newProperties] that are about to be tracked, stop tracking any properties that are no longer relevant and
+  /// start any properties that are newly relevant. [exposeFutures] will be filled with futures for each newly-tracked
+  /// property that will complete when the property is exposed, and [callback] will be called in the future whenever a
+  /// property's value changes.
+  void _trackRelevantProperties<T>({
+    required List<WidgetControlledUnrealProperty<T>?> trackedProperties,
+    required List<UnrealProperty> oldProperties,
+    required List<UnrealProperty> newProperties,
+    required List<Future<void>> exposeFutures,
+    required Function(int trackedPropertyVersion, int index, dynamic newValue) callback,
+  }) {
+    assert(oldProperties.length == trackedProperties.length);
+
+    // Stop tracking all properties. We either aren't tracking them anymore, or we may need to update their
+    // propertyIndex if the number tracked has changed size, in which case we need to recreate the callback anyway.
+    for (int propertyIndex = 0; propertyIndex < oldProperties.length; ++propertyIndex) {
+      final WidgetControlledUnrealProperty<T>? property = trackedProperties[propertyIndex];
+
+      if (property != null) {
+        _propertyManager.stopTrackingProperty(property.trackedId, property.changedCallback);
+      }
+    }
+
+    oldProperties.clear();
+    oldProperties.addAll(newProperties);
+
+    trackedProperties.length = newProperties.length;
+
+    // Track each property we want to control with the remote property manager.
+    for (int propertyIndex = 0; propertyIndex < newProperties.length; ++propertyIndex) {
+      final UnrealProperty property = newProperties[propertyIndex];
+
+      // Create a callback function for this property
+      changedCallback(newValue) => callback(_latestTrackedPropertyVersion, propertyIndex, newValue);
+
+      // Start tracking the property in the property manager
+      final TrackedPropertyId trackedId =
+          _propertyManager.trackProperty(property, changedCallback, conversionMetadata: conversionMetadata);
+
+      final newProperty = WidgetControlledUnrealProperty<T>(
+        trackedId: trackedId,
+        propertyIndex: propertyIndex,
+        changedCallback: changedCallback,
+        value: _propertyManager.getTrackedPropertyValue(trackedId),
+        ownerName: property.objectName,
+      );
+      trackedProperties[propertyIndex] = newProperty;
+
+      // Add to the list of properties we're waiting for
+      final Future<void> exposeFuture = _propertyManager.waitForProperty(trackedId);
+      exposeFutures.add(exposeFuture);
+
+      exposeFuture.then((_) {
+        newProperty.value = _propertyManager.getTrackedPropertyValue(trackedId);
+      }).onError((error, stackTrace) {
+        _log.severe(
+            'Error while exposing property ${property.propertyName} on ${property.objectName}', error, stackTrace);
+      });
+    }
+  }
+
+  /// Called when all properties have been subscribed and the widget is ready to be interacted with.
+  void _onPropertiesReady() {
+    if (_properties.isNotEmpty) {
+      // Properties could differ in their metadata, but for now we'll assume they share the same values
+      final WidgetControlledUnrealProperty<PropertyType>? firstProperty = _properties[0];
+
+      if (firstProperty == null) {
+        assert(false, 'If an Unreal widget is ready, all properties should be non-null');
+        return;
+      }
+
+      _propertyMetadata = _propertyManager.getTrackedPropertyMetaData(firstProperty.trackedId);
+    }
+
+    setState(() {
+      _bIsReady = true;
+    });
+  }
+
+  /// Called when the property manager reports a [newValue] for the property with the given [index] in the widget.
+  /// If [trackedPropertyVersion] does not match [_latestTrackedPropertyVersion], the change will be ignored.
+  void _handleOnManagedValueChanged(int trackedPropertyVersion, int index, dynamic newValue) {
+    if (trackedPropertyVersion != _latestTrackedPropertyVersion) {
+      return;
+    }
+
+    setState(() {
+      _properties[index]?.value = newValue;
+    });
+  }
+
+  /// Called when an enable value for one of the properties changes.
+  void _handleOnEnableValueChanged(int trackedPropertyVersion, int index, dynamic newValue) {}
+}
diff --git a/Engine/Extras/VirtualProduction/EpicStageApp/lib/widgets/screens/connect/connect.dart b/Engine/Extras/VirtualProduction/EpicStageApp/lib/widgets/screens/connect/connect.dart
new file mode 100644
--- /dev/null
+++ b/Engine/Extras/VirtualProduction/EpicStageApp/lib/widgets/screens/connect/connect.dart
@@ -0,0 +1,220 @@
+// Copyright Epic Games, Inc. All Rights Reserved.
+
+import 'package:flutter/material.dart';
+import 'package:flutter_gen/gen_l10n/app_localizations.dart';
+import 'package:provider/provider.dart';
+
+import '../../../../utilities/constants.dart';
+import '../../../../widgets/elements/empty_placeholder.dart';
+import '../../../../widgets/elements/layout/card.dart';
+import '../../../../widgets/screens/connect/views/connection_list.dart';
+import '../../../../widgets/screens/main/toolbar/main_screen_toolbar.dart';
+import '../../../models/engine_connection.dart';
+import '../../../models/navigator_keys.dart';
+import '../../../models/unreal_engine_beacon.dart';
+import '../../../utilities/debug_utilities.dart';
+import '../../../utilities/net_utilities.dart';
+import '../../../utilities/unreal_colors.dart';
+import '../../elements/epic_icon_button.dart';
+import '../../elements/modal.dart';
+import 'mixins/connect_mixin.dart';
+import 'views/manual_connect_form.dart';
+import 'views/quick_action_grid.dart';
+
+/// A screen that uses UDP multicast to search for any local Unreal Engine instances and lists them for the user to
+/// choose from.
+class ConnectScreen extends StatefulWidget {
+  const ConnectScreen({Key? key}) : super(key: key);
+
+  static const String route = '/connect';
+
+  @override
+  _ConnectScreenState createState() => _ConnectScreenState();
+}
+
+class _ConnectScreenState extends State<ConnectScreen> with WidgetsBindingObserver, ConnectMixin {
+  /// Beacon which will be used to detect instances of Unreal Engine.
+  late final UnrealEngineBeacon _engineBeacon;
+
+  /// Data for the last engine we were connected to, if any.
+  ConnectionData? _lastConnection;
+
+  /// List of valid engine instances to connect to.
+  List<ConnectionData> _connections = [];
+
+  /// Called when the beacon has failed for too long.
+  void _onBeaconFailure() {
+    // If there is an opened route that allows a user go back or pop the Navigation stack once,
+    // we likely should not be showing this alert then.
+    if (!Navigator.of(context).canPop()) {
+      showDebugAlert(AppLocalizations.of(context)!.connectScreenBeaconFailedMessage);
+    }
+  }
+
+  @override
+  void initState() {
+    super.initState();
+
+    WidgetsBinding.instance.addObserver(this);
+
+    _engineBeacon = UnrealEngineBeacon(
+      context: context,
+      onBeaconFailure: _onBeaconFailure,
+      onEngineInstancesChanged: () => setState(() {}),
+    );
+
+    Provider.of<EngineConnectionManager>(context, listen: false)
+        .getLastConnectionData()
+        .then(_receiveLastConnectionData);
+  }
+
+  @override
+  void dispose() {
+    super.dispose();
+
+    WidgetsBinding.instance.removeObserver(this);
+
+    _engineBeacon.dispose();
+  }
+
+  @override
+  void didChangeAppLifecycleState(AppLifecycleState state) {
+    switch (state) {
+      case AppLifecycleState.resumed:
+        _engineBeacon.resume();
+        break;
+
+      case AppLifecycleState.paused:
+      case AppLifecycleState.detached:
+        // If the app is in the background, stop sending beacon messages
+        _engineBeacon.pause();
+        break;
+
+      default:
+        break;
+    }
+  }
+
+  @override
+  Widget build(BuildContext context) {
+    return Scaffold(
+      resizeToAvoidBottomInset: false,
+      appBar: const _ConnectScreenToolbar(),
+      body: Container(
+        color: UnrealColors.gray14,
+        padding: EdgeInsets.only(left: cardMargin, right: cardMargin, bottom: cardMargin),
+        child: Container(
+          width: MediaQuery.of(context).size.width,
+          padding: EdgeInsets.all(cardMargin),
+          decoration: BoxDecoration(
+            color: Theme.of(context).colorScheme.background,
+            borderRadius: BorderRadius.circular(outerCornerRadius),
+          ),
+          child: Row(
+            mainAxisSize: MainAxisSize.min,
+            children: [
+              Card(
+                child: IntrinsicWidth(
+                  child: Column(
+                    children: [
+                      CardLargeHeader(
+                        title: AppLocalizations.of(context)!.connectScreenQuickConnectPanelTitle,
+                        iconPath: 'assets/images/icons/light_card.svg',
+                      ),
+                      QuickActionGrid(recentConnections: _connections),
+                    ],
+                  ),
+                ),
+              ),
+              SizedBox(width: cardMargin),
+              Expanded(
+                child: Card(
+                  child: Column(
+                    crossAxisAlignment: CrossAxisAlignment.stretch,
+                    children: [
+                      CardLargeHeader(
+                        title: AppLocalizations.of(context)!.connectScreenAllConnectionsPanelTitle,
+                        iconPath: 'assets/images/icons/unreal_u_logo.svg',
+                      ),
+                      Expanded(
+                        child: _engineBeacon.connections.length > 0
+                            ? ConnectionList(
+                                connections: _engineBeacon.connections,
+                              )
+                            : EmptyPlaceholder(
+                                message: AppLocalizations.of(context)!.connectScreenAllConnectionsPanelEmptyMessage,
+                                button: EpicWideButton(
+                                  text: AppLocalizations.of(context)!.connectScreenAllConnectionsPanelEmptyButtonLabel,
+                                  iconPath: 'assets/images/icons/plus.svg',
+                                  color: UnrealColors.highlightGreen,
+                                  onPressed: _showManualConnectDialog,
+                                ),
+                              ),
+                      ),
+                    ],
+                  ),
+                ),
+              ),
+            ],
+          ),
+        ),
+      ),
+    );
+  }
+
+  /// Called when the data for the last connection is loaded.
+  void _receiveLastConnectionData(ConnectionData? connectionData) {
+    if (connectionData != null) {
+      setState(() {
+        _lastConnection = connectionData;
+        _connections.removeWhere((existingConnection) => _isLastConnection(existingConnection));
+      });
+      _connections.add(connectionData);
+    }
+  }
+
+  /// Check if the given [connectionData] is the last connection we made.
+  bool _isLastConnection(ConnectionData connectionData) {
+    // We can't compare on UUID because the engine may have restarted, so check everything else
+    return _lastConnection != null &&
+        connectionData.websocketAddress == _lastConnection!.websocketAddress &&
+        connectionData.websocketPort == _lastConnection!.websocketPort &&
+        connectionData.name == _lastConnection!.name;
+  }
+
+  /// Show the dialog for manually connecting to the engine.
+  void _showManualConnectDialog() {
+    GenericModalDialogRoute.showDialog(
+      context: rootNavigatorKey.currentContext!,
+      builder: (context) => ModalDialogCard(
+        child: ManualConnectForm(connect: connect),
+      ),
+    );
+  }
+}
+
+/// Toolbar shown only on the connect screen.
+class _ConnectScreenToolbar extends StatelessWidget implements PreferredSizeWidget {
+  const _ConnectScreenToolbar({Key? key}) : super(key: key);
+
+  @override
+  Size get preferredSize => const Size.fromHeight(48);
+
+  @override
+  Widget build(BuildContext context) {
+    return Container(
+      padding: const EdgeInsets.only(right: 8),
+      decoration: BoxDecoration(
+        color: Theme.of(context).colorScheme.surface,
+      ),
+      height: preferredSize.height,
+      child: Row(
+        mainAxisAlignment: MainAxisAlignment.end,
+        crossAxisAlignment: CrossAxisAlignment.center,
+        children: [
+          SettingsButton(),
+        ],
+      ),
+    );
+  }
+}
diff --git a/Engine/Extras/VirtualProduction/EpicStageApp/lib/widgets/screens/connect/mixins/connect_mixin.dart b/Engine/Extras/VirtualProduction/EpicStageApp/lib/widgets/screens/connect/mixins/connect_mixin.dart
new file mode 100644
--- /dev/null
+++ b/Engine/Extras/VirtualProduction/EpicStageApp/lib/widgets/screens/connect/mixins/connect_mixin.dart
@@ -0,0 +1,183 @@
+// Copyright Epic Games, Inc. All Rights Reserved.
+
+import 'dart:async';
+
+import 'package:flutter/material.dart';
+import 'package:flutter_gen/gen_l10n/app_localizations.dart';
+import 'package:provider/provider.dart';
+
+import '../../../../models/engine_connection.dart';
+import '../../../../models/navigator_keys.dart';
+import '../../../../models/settings/selected_actor_settings.dart';
+import '../../../../models/unreal_actor_manager.dart';
+import '../../../../models/unreal_types.dart';
+import '../../../../utilities/constants.dart';
+import '../../../../utilities/net_utilities.dart';
+import '../../../../utilities/preferences_bundle.dart';
+import '../../../elements/asset_icon.dart';
+import '../../../elements/epic_icon_button.dart';
+import '../../../elements/modal.dart';
+import '../../main/stage_app_main_screen.dart';
+import '../views/error_modal.dart';
+import '../views/n_display_selector_dialog.dart';
+import '../views/passphrase_modal.dart';
+
+/// Connection Mixin for managing states related to UE connection.
+mixin ConnectMixin<T extends StatefulWidget> on State<T> {
+  /// A state instance for managing, listening/watching and disposing listeners for various class names.
+  late final UnrealActorManager _actorManager;
+
+  /// A state instance used to manage selected actors.
+  late final SelectedActorSettings _actorSettings;
+
+  /// A state instance used to establish and manage connection to UE.
+  late final EngineConnectionManager _connectionManager;
+
+  /// Whether the [_buildSpinner] is still currently being shown to the user.
+  bool? bIsConnecting = false;
+
+  /// Set of available actors found in the connected UE instance.
+  Set<UnrealObject> actors = {};
+
+  @override
+  void initState() {
+    _actorManager = Provider.of<UnrealActorManager>(context, listen: false);
+    _actorSettings = Provider.of<SelectedActorSettings>(context, listen: false);
+    _connectionManager = Provider.of<EngineConnectionManager>(context, listen: false);
+    super.initState();
+  }
+
+  /// Callback function for connecting to an instance of UE, using [data].
+  /// If [bShowPassphraseModalOnReject] is true, this will show a modal dialog to enter the passphrase if it's rejected.
+  Future<EngineConnectionResult> connect(
+    ConnectionData data, {
+    bool bShowPassphraseModalOnReject = true,
+  }) async {
+    _buildSpinner(context, data);
+    setState(() => bIsConnecting = true);
+
+    final EngineConnectionResult result = await _connectionManager.connect(data);
+
+    if (result != EngineConnectionResult.success) {
+      if (bIsConnecting == true) {
+        Navigator.of(context).pop();
+
+        if (result == EngineConnectionResult.passphraseRejected) {
+          if (bShowPassphraseModalOnReject) {
+            GenericModalDialogRoute.showDialog(
+              context: context,
+              builder: (context) => PassphraseModalDialog(),
+            );
+          }
+        } else {
+          _buildError(context, data);
+        }
+      }
+
+      return result;
+    }
+
+    final Set<UnrealObject> actors = await _actorManager.getInitialActorsOfClass(nDisplayRootActorClassName);
+
+    final store = Provider.of<PreferencesBundle>(context, listen: false).persistent;
+
+    if (actors.length == 1) {
+      // Use the only available root actor
+      _actorSettings.displayClusterRootPath.setValue(actors.first.path);
+    } else if (actors.length > 1) {
+      // Multiple root actors available, so check if we'd previously selected one of them
+      final String previousDisplayRootActor =
+          store.getString('common.displayClusterRootPath', defaultValue: '').getValue();
+
+      if (!actors.any((element) => element.path == previousDisplayRootActor)) {
+        // If the progress dialog is still showing we want to stop it before showing the NDisplayDialogSelector.
+        if (bIsConnecting == true) {
+          Navigator.pop(context);
+        }
+
+        final UnrealObject? newActor = await NDisplaySelectorDialog.showIfNotOpen(context);
+
+        if (newActor == null) {
+          // User cancelled
+          _connectionManager.disconnect();
+          return EngineConnectionResult.cancelled;
+        }
+
+        _actorSettings.displayClusterRootPath.setValue(newActor.path);
+      }
+    }
+
+    Navigator.of(context).pushReplacementNamed(StageAppMainScreen.route);
+
+    return result;
+  }
+
+  /// Show Circular progress indicator in the current [context] when connecting to UE using [data].
+  _buildSpinner(BuildContext context, ConnectionData data) async {
+    final route = GenericModalDialogRoute(
+      bResizeToAvoidBottomInset: true,
+      // To avoid confusion, only allow dismissal by pressing the Cancel button
+      bIsBarrierDismissible: false,
+      builder: (_) => (context as Element).debugIsActive
+          ? ModalDialogCard(
+              child: Container(
+                width: MediaQuery.of(rootNavigatorKey.currentContext!).size.width * 0.5,
+                padding: EdgeInsets.all(10),
+                child: Column(
+                  mainAxisSize: MainAxisSize.min,
+                  children: [
+                    Row(
+                      mainAxisAlignment: MainAxisAlignment.center,
+                      mainAxisSize: MainAxisSize.min,
+                      children: [
+                        AssetIcon(
+                          path: 'assets/images/icons/unreal_u_logo.svg',
+                          size: 24,
+                        ),
+                        SizedBox(width: 16),
+                        Text(
+                          AppLocalizations.of(context)!.connectScreenConnectDialogTitle,
+                          style: Theme.of(context).textTheme.displayLarge,
+                        ),
+                      ],
+                    ),
+                    const SizedBox(height: 16),
+                    Text(AppLocalizations.of(context)!
+                        .connectScreenConnectDialogMessage('${data.websocketAddress.address}:${data.websocketPort}')),
+                    const SizedBox(height: 32),
+                    SizedBox(
+                      child: const CircularProgressIndicator(strokeWidth: 8.5),
+                      height: 175,
+                      width: 175,
+                    ),
+                    Row(mainAxisAlignment: MainAxisAlignment.end, children: [
+                      EpicLozengeButton(
+                        label: AppLocalizations.of(context)!.menuButtonCancel,
+                        color: Theme.of(context).colorScheme.secondary,
+                        onPressed: () {
+                          Navigator.pop(context, false);
+                          _connectionManager.disconnect();
+                          setState(() => bIsConnecting = false);
+                        },
+                      ),
+                    ]),
+                  ],
+                ),
+              ),
+            )
+          : SizedBox(),
+    );
+
+    return Navigator.of(context).push(route);
+  }
+
+  /// Show error message dialog/modal in the current [context], when and if connection to UE fails using [data].
+  _buildError(BuildContext context, ConnectionData data) async {
+    return GenericModalDialogRoute.showDialog(
+      context: context,
+      builder: (_) => ModalDialogCard(
+        child: ErrorModal(data: data, reconnect: () => connect(data)),
+      ),
+    );
+  }
+}
diff --git a/Engine/Extras/VirtualProduction/EpicStageApp/lib/widgets/screens/connect/views/connection_list.dart b/Engine/Extras/VirtualProduction/EpicStageApp/lib/widgets/screens/connect/views/connection_list.dart
new file mode 100644
--- /dev/null
+++ b/Engine/Extras/VirtualProduction/EpicStageApp/lib/widgets/screens/connect/views/connection_list.dart
@@ -0,0 +1,36 @@
+// Copyright Epic Games, Inc. All Rights Reserved.
+
+import 'package:flutter/material.dart';
+
+import '../../../../utilities/net_utilities.dart';
+import 'connection_tile.dart';
+
+/// A widget showing a list of available UE connection instances.
+class ConnectionList extends StatelessWidget {
+  const ConnectionList({Key? key, required this.connections}) : super(key: key);
+
+  /// List of UE connection instances to be rendered.
+  final List<ConnectionData> connections;
+
+  @override
+  Widget build(BuildContext context) {
+    return SingleChildScrollView(
+      padding: EdgeInsets.all(32),
+      child: Wrap(
+        direction: Axis.horizontal,
+        alignment: WrapAlignment.center,
+        runAlignment: WrapAlignment.start,
+        crossAxisAlignment: WrapCrossAlignment.start,
+        spacing: 16,
+        runSpacing: 32,
+        children: [
+          for (final ConnectionData connection in connections)
+            ConnectionTile(
+              connectionData: connection,
+              bFilled: false,
+            ),
+        ],
+      ),
+    );
+  }
+}
diff --git a/Engine/Extras/VirtualProduction/EpicStageApp/lib/widgets/screens/connect/views/connection_tile.dart b/Engine/Extras/VirtualProduction/EpicStageApp/lib/widgets/screens/connect/views/connection_tile.dart
new file mode 100644
--- /dev/null
+++ b/Engine/Extras/VirtualProduction/EpicStageApp/lib/widgets/screens/connect/views/connection_tile.dart
@@ -0,0 +1,118 @@
+// Copyright Epic Games, Inc. All Rights Reserved.
+
+import 'package:flutter/material.dart';
+import 'package:flutter_gen/gen_l10n/app_localizations.dart';
+
+import '../../../../utilities/constants.dart';
+import '../../../../utilities/net_utilities.dart';
+import '../../../../utilities/unreal_colors.dart';
+import '../../../elements/asset_icon.dart';
+import '../mixins/connect_mixin.dart';
+
+/// Connection tile/item used to visually represent available instances of UE connections.
+class ConnectionTile extends StatefulWidget {
+  const ConnectionTile({
+    Key? key,
+    required this.connectionData,
+    this.bFilled = true,
+    this.bPrimary = false,
+  }) : super(key: key);
+
+  /// [connectionData] essential to connecting with an instance of UE.
+  final ConnectionData connectionData;
+
+  /// Whether the item is filled or has only border, default to `true`.
+  final bool bFilled;
+
+  /// Whether the item should have the primary color, default to `false`.
+  final bool bPrimary;
+
+  @override
+  State<ConnectionTile> createState() => _ConnectionTileState();
+}
+
+class _ConnectionTileState extends State<ConnectionTile> with ConnectMixin {
+  /// State instance for [widget.connectionData], for easy access across child widgets.
+  late final ConnectionData data;
+
+  @override
+  void initState() {
+    data = widget.connectionData;
+    super.initState();
+  }
+
+  @override
+  Widget build(BuildContext context) {
+    late final Color fillColor;
+    if (widget.bFilled) {
+      if (widget.bPrimary) {
+        fillColor = Theme.of(context).colorScheme.primary;
+      } else {
+        fillColor = Theme.of(context).colorScheme.secondary;
+      }
+    } else {
+      fillColor = Colors.transparent;
+    }
+
+    late final Color textColor = widget.bFilled ? Theme.of(context).colorScheme.onPrimary : UnrealColors.gray56;
+
+    return InkWell(
+      onTap: () => connect(data),
+      child: Container(
+        width: 170,
+        height: 166,
+        padding: EdgeInsets.symmetric(horizontal: 10),
+        decoration: BoxDecoration(
+          borderRadius: BorderRadius.circular(cardCornerRadius),
+          color: fillColor,
+          border: widget.bFilled
+              ? null
+              : Border.all(
+                  color: Theme.of(context).colorScheme.onSurface.withOpacity(0.2),
+                  width: 2,
+                ),
+        ),
+        child: Column(
+          crossAxisAlignment: CrossAxisAlignment.center,
+          mainAxisAlignment: MainAxisAlignment.center,
+          children: [
+            AssetIcon(
+              path: 'assets/images/icons/unreal_u.svg',
+              size: 50,
+              color: widget.bFilled ? UnrealColors.white : UnrealColors.gray75,
+            ),
+            SizedBox(height: 16),
+            Text(
+              '${widget.connectionData.websocketAddress.address}:${widget.connectionData.websocketPort}',
+              style: Theme.of(context).textTheme.titleLarge!.copyWith(
+                    letterSpacing: 0.5,
+                    color: textColor,
+                  ),
+              textAlign: TextAlign.center,
+              overflow: TextOverflow.ellipsis,
+            ),
+            SizedBox(height: 5),
+            Text(
+              widget.connectionData.name,
+              style: Theme.of(context).textTheme.labelMedium!.copyWith(
+                    color: textColor,
+                  ),
+              textAlign: TextAlign.center,
+              overflow: TextOverflow.ellipsis,
+            ),
+            SizedBox(height: 5),
+            widget.bFilled
+                ? Text(
+                    AppLocalizations.of(context)!.connectScreenMostRecentButtonLabel,
+                    style: Theme.of(context).textTheme.labelMedium!.copyWith(
+                          color: textColor,
+                          fontStyle: FontStyle.italic,
+                        ),
+                  )
+                : const SizedBox(),
+          ],
+        ),
+      ),
+    );
+  }
+}
diff --git a/Engine/Extras/VirtualProduction/EpicStageApp/lib/widgets/screens/connect/views/error_modal.dart b/Engine/Extras/VirtualProduction/EpicStageApp/lib/widgets/screens/connect/views/error_modal.dart
new file mode 100644
--- /dev/null
+++ b/Engine/Extras/VirtualProduction/EpicStageApp/lib/widgets/screens/connect/views/error_modal.dart
@@ -0,0 +1,75 @@
+// Copyright Epic Games, Inc. All Rights Reserved.
+
+import 'package:flutter/material.dart';
+import 'package:flutter_gen/gen_l10n/app_localizations.dart';
+
+import '../../../../utilities/net_utilities.dart';
+import '../../../elements/asset_icon.dart';
+import '../../../elements/epic_icon_button.dart';
+import '../../../elements/parsed_rich_text.dart';
+
+/// Error modal for when there's an error connecting to an instance of UE.
+class ErrorModal extends StatelessWidget {
+  const ErrorModal({Key? key, required this.data, this.reconnect}) : super(key: key);
+
+  /// connection data/details used for connecting to UE when the error occurred.
+  final ConnectionData data;
+
+  /// callback function for reconnecting to an instance of UE after an error/failure using [data].
+  final VoidCallback? reconnect;
+
+  @override
+  Widget build(BuildContext context) {
+    return Container(
+      width: MediaQuery.of(context).size.height * 0.5,
+      padding: EdgeInsets.symmetric(horizontal: 16, vertical: 8),
+      child: Column(
+        mainAxisSize: MainAxisSize.min,
+        children: [
+          Row(
+            mainAxisAlignment: MainAxisAlignment.center,
+            children: [
+              AssetIcon(
+                path: 'assets/images/icons/alert_triangle_large.svg',
+                size: 24,
+              ),
+              SizedBox(width: 10),
+              Flexible(
+                child: Text(
+                  AppLocalizations.of(context)!.connectScreenFailedDialogTitle,
+                  style: Theme.of(context).textTheme.displayLarge,
+                  overflow: TextOverflow.fade,
+                ),
+              ),
+            ],
+          ),
+          const SizedBox(height: 24),
+          ParsedRichText(
+            AppLocalizations.of(context)!
+                .connectScreenFailedDialogMessage('${data.websocketAddress.address}:${data.websocketPort}'),
+            style: Theme.of(context).textTheme.bodyMedium,
+          ),
+          const SizedBox(height: 30),
+          Row(
+            mainAxisAlignment: MainAxisAlignment.end,
+            children: [
+              EpicLozengeButton(
+                onPressed: () => Navigator.of(context).pop(),
+                label: AppLocalizations.of(context)!.menuButtonCancel,
+                color: Colors.transparent,
+              ),
+              SizedBox(width: 16),
+              EpicLozengeButton(
+                onPressed: () {
+                  Navigator.of(context).pop();
+                  reconnect?.call();
+                },
+                label: AppLocalizations.of(context)!.menuButtonRetry,
+              ),
+            ],
+          ),
+        ],
+      ),
+    );
+  }
+}
diff --git a/Engine/Extras/VirtualProduction/EpicStageApp/lib/widgets/screens/connect/views/manual_connect_form.dart b/Engine/Extras/VirtualProduction/EpicStageApp/lib/widgets/screens/connect/views/manual_connect_form.dart
new file mode 100644
--- /dev/null
+++ b/Engine/Extras/VirtualProduction/EpicStageApp/lib/widgets/screens/connect/views/manual_connect_form.dart
@@ -0,0 +1,147 @@
+// Copyright Epic Games, Inc. All Rights Reserved.
+
+import 'dart:io';
+
+import 'package:flutter/material.dart';
+import 'package:flutter_gen/gen_l10n/app_localizations.dart';
+import 'package:uuid/uuid.dart';
+
+import '../../../../utilities/constants.dart';
+import '../../../../utilities/net_utilities.dart';
+import '../../../elements/epic_icon_button.dart';
+import '../../../elements/modal.dart';
+
+/// Manual connection form for manually connecting to an instance of UE.
+class ManualConnectForm extends StatefulWidget {
+  const ManualConnectForm({Key? key, required this.connect}) : super(key: key);
+
+  /// A function that allows us connect to an instance of UE.
+  final Function(ConnectionData) connect;
+
+  @override
+  State<ManualConnectForm> createState() => _ManualConnectFormState();
+}
+
+class _ManualConnectFormState extends State<ManualConnectForm> {
+  /// Global key for managing form state & validating form input.
+  final _formKey = GlobalKey<FormState>();
+
+  /// Input controller for IP address on the form.
+  final _ipTextController = TextEditingController(text: '');
+
+  /// Input controller for PORT on the form.
+  final _portTextController = TextEditingController(text: '30020');
+
+  @override
+  Widget build(BuildContext context) {
+    return Form(
+      key: _formKey,
+      child: IntrinsicHeight(
+        child: IntrinsicWidth(
+          child: Column(
+            children: [
+              ModalDialogTitle(title: AppLocalizations.of(context)!.connectScreenConnectionFormTitle),
+              ModalDialogSection(
+                child: TextFormField(
+                  cursorWidth: 1,
+                  style: Theme.of(context).textTheme.bodyMedium,
+                  keyboardType: const TextInputType.numberWithOptions(
+                    decimal: true,
+                    signed: false,
+                  ),
+                  decoration: collapsedInputDecoration.copyWith(
+                    hintText: AppLocalizations.of(context)!.hintIPAddress,
+                  ),
+                  controller: _ipTextController,
+                  validator: _validateAddress,
+                ),
+              ),
+              ModalDialogSection(
+                child: TextFormField(
+                  cursorWidth: 1,
+                  style: Theme.of(context).textTheme.bodyMedium,
+                  keyboardType: const TextInputType.numberWithOptions(
+                    decimal: false,
+                    signed: false,
+                  ),
+                  decoration: collapsedInputDecoration.copyWith(
+                    hintText: AppLocalizations.of(context)!.hintNetworkPort,
+                  ),
+                  controller: _portTextController,
+                  validator: _validatePort,
+                ),
+              ),
+              ModalDialogSection(
+                child: Row(
+                  mainAxisAlignment: MainAxisAlignment.end,
+                  children: [
+                    EpicLozengeButton(
+                      label: AppLocalizations.of(context)!.menuButtonCancel,
+                      color: Colors.transparent,
+                      onPressed: () => Navigator.of(context).pop(),
+                    ),
+                    EpicLozengeButton(
+                      label: AppLocalizations.of(context)!.menuButtonOK,
+                      onPressed: _onConnect,
+                    ),
+                  ],
+                ),
+              ),
+            ],
+          ),
+        ),
+      ),
+    );
+  }
+
+  /// Callback for when the connect button is clicked or pressed.
+  void _onConnect() {
+    if (_formKey.currentState?.validate() != true) {
+      return;
+    }
+
+    Navigator.of(context).pop(false);
+
+    ConnectionData connection = ConnectionData(
+      uuid: UuidValue(const Uuid().v1()),
+      name: AppLocalizations.of(context)!.connectScreenManualConnectionLabel,
+      websocketAddress: InternetAddress(_ipTextController.text),
+      websocketPort: int.parse(_portTextController.text),
+    );
+
+    /// Form input validating has passed, we are calling [widget.connect] to initiate a connection with UE using
+    /// [connection]
+    widget.connect(connection);
+  }
+
+  /// String validation for validating IpAddresses.
+  String? _validateAddress(String? text) {
+    if (text == null) {
+      return AppLocalizations.of(context)!.formErrorInvalidIPAddress;
+    }
+
+    try {
+      InternetAddress(text);
+    } catch (e) {
+      return AppLocalizations.of(context)!.formErrorInvalidIPAddress;
+    }
+    return null;
+  }
+
+  /// String validation for validating ports.
+  String? _validatePort(String? text) {
+    if (text == null) {
+      return AppLocalizations.of(context)!.formErrorInvalidNetworkPort;
+    }
+
+    try {
+      final port = int.parse(text);
+      if (port > 65535) {
+        return AppLocalizations.of(context)!.formErrorInvalidNetworkPort;
+      }
+    } catch (e) {
+      return AppLocalizations.of(context)!.formErrorInvalidNetworkPort;
+    }
+    return null;
+  }
+}
diff --git a/Engine/Extras/VirtualProduction/EpicStageApp/lib/widgets/screens/connect/views/manual_connect_tile.dart b/Engine/Extras/VirtualProduction/EpicStageApp/lib/widgets/screens/connect/views/manual_connect_tile.dart
new file mode 100644
--- /dev/null
+++ b/Engine/Extras/VirtualProduction/EpicStageApp/lib/widgets/screens/connect/views/manual_connect_tile.dart
@@ -0,0 +1,80 @@
+// Copyright Epic Games, Inc. All Rights Reserved.
+
+import 'dart:async';
+
+import 'package:flutter/material.dart';
+import 'package:flutter_gen/gen_l10n/app_localizations.dart';
+
+import '../../../../models/navigator_keys.dart';
+import '../../../../utilities/constants.dart';
+import '../../../elements/asset_icon.dart';
+import '../../../elements/modal.dart';
+import '../mixins/connect_mixin.dart';
+import 'manual_connect_form.dart';
+
+/// Button or Tile used for manually connecting to an instance of UE.
+class ManualConnectTile extends StatefulWidget {
+  const ManualConnectTile({Key? key}) : super(key: key);
+
+  @override
+  State<ManualConnectTile> createState() => _ManualConnectTileState();
+}
+
+class _ManualConnectTileState extends State<ManualConnectTile> with ConnectMixin {
+  /// Show the input dialog/modal for connecting manually to an instance of UE in the current [context].
+  Future<dynamic> _showManualConnectionInput(BuildContext context) async {
+    final route = GenericModalDialogRoute(
+      bResizeToAvoidBottomInset: true,
+      builder: (_) => ModalDialogCard(
+        child: ManualConnectForm(connect: connect),
+      ),
+    );
+
+    return Navigator.of(rootNavigatorKey.currentContext!, rootNavigator: true).push(route);
+  }
+
+  @override
+  Widget build(BuildContext context) {
+    return InkWell(
+      onTap: () => _showManualConnectionInput(context),
+      child: Container(
+        width: 170,
+        height: 166,
+        decoration: BoxDecoration(
+          borderRadius: BorderRadius.circular(cardCornerRadius),
+          color: Theme.of(context).colorScheme.secondary,
+        ),
+        padding: EdgeInsets.all(10),
+        child: Column(
+          crossAxisAlignment: CrossAxisAlignment.center,
+          mainAxisAlignment: MainAxisAlignment.center,
+          mainAxisSize: MainAxisSize.min,
+          children: [
+            AssetIcon(
+              path: 'assets/images/icons/unreal_u.svg',
+              size: 50,
+            ),
+            SizedBox(height: 16),
+            Text(
+              AppLocalizations.of(context)!.connectScreenNewConnectionButtonTitle,
+              style: Theme.of(context).textTheme.titleLarge!.copyWith(
+                    letterSpacing: 0.5,
+                    color: Theme.of(context).colorScheme.onPrimary,
+                    height: 1.5,
+                  ),
+              textAlign: TextAlign.center,
+            ),
+            SizedBox(height: 5),
+            Text(
+              AppLocalizations.of(context)!.connectScreenNewConnectionButtonLabel,
+              style: Theme.of(context).textTheme.labelMedium!.copyWith(
+                    color: Theme.of(context).colorScheme.onPrimary,
+                    fontStyle: FontStyle.italic,
+                  ),
+            ),
+          ],
+        ),
+      ),
+    );
+  }
+}
diff --git a/Engine/Extras/VirtualProduction/EpicStageApp/lib/widgets/screens/connect/views/n_display_selector_dialog.dart b/Engine/Extras/VirtualProduction/EpicStageApp/lib/widgets/screens/connect/views/n_display_selector_dialog.dart
new file mode 100644
--- /dev/null
+++ b/Engine/Extras/VirtualProduction/EpicStageApp/lib/widgets/screens/connect/views/n_display_selector_dialog.dart
@@ -0,0 +1,207 @@
+// Copyright Epic Games, Inc. All Rights Reserved.
+
+import 'package:flutter/material.dart';
+import 'package:flutter_gen/gen_l10n/app_localizations.dart';
+import 'package:provider/provider.dart';
+
+import '../../../../models/unreal_actor_manager.dart';
+import '../../../../models/unreal_types.dart';
+import '../../../../utilities/constants.dart';
+import '../../../../utilities/unreal_colors.dart';
+import '../../../elements/asset_icon.dart';
+import '../../../elements/epic_icon_button.dart';
+import '../../../elements/modal.dart';
+
+/// Displays the list of currently available actors which users can select from.
+class NDisplaySelectorDialog extends StatefulWidget {
+  const NDisplaySelectorDialog({Key? key, this.bIsAlreadyConnected = false}) : super(key: key);
+
+  /// Whether the user was already connected to the engine when this was shown.
+  final bool bIsAlreadyConnected;
+
+  /// The future for when the currently active dialog closes.
+  static Future<UnrealObject?>? activeDialogFuture;
+
+  /// Helper function to show this dialog while guaranteeing that only one is open at a time.
+  /// If [bIsAlreadyConnected] is true, adjust buttons and menu behaviour to reflect that the next step will not be a
+  /// new engine connection.
+  static Future<UnrealObject?> showIfNotOpen(
+    BuildContext context, {
+    bool bIsAlreadyConnected = false,
+  }) {
+    if (activeDialogFuture == null) {
+      activeDialogFuture = GenericModalDialogRoute.showDialog(
+        context: context,
+        builder: (context) => NDisplaySelectorDialog(
+          bIsAlreadyConnected: bIsAlreadyConnected,
+        ),
+        bIsBarrierDismissible: !bIsAlreadyConnected,
+      );
+
+      activeDialogFuture!.then((_) => activeDialogFuture = null);
+    }
+
+    return activeDialogFuture!;
+  }
+
+  @override
+  State<NDisplaySelectorDialog> createState() => _NDisplaySelectorDialogState();
+}
+
+class _NDisplaySelectorDialogState extends State<NDisplaySelectorDialog> {
+  /// currently selected actor, value will null when no actor is selected.
+  UnrealObject? _selectedActor = null;
+
+  /// Manager used to retrieve the list of root actors.
+  late final UnrealActorManager _actorManager;
+
+  @override
+  void initState() {
+    super.initState();
+
+    _actorManager = Provider.of<UnrealActorManager>(context, listen: false);
+    _actorManager.watchClassName(nDisplayRootActorClassName, _onRootActorUpdate);
+  }
+
+  @override
+  void dispose() {
+    _actorManager.stopWatchingClassName(nDisplayRootActorClassName, _onRootActorUpdate);
+
+    super.dispose();
+  }
+
+  @override
+  Widget build(BuildContext context) {
+    return ModalDialogCard(
+      child: Container(
+        width: MediaQuery.of(context).size.width * .4,
+        padding: EdgeInsets.symmetric(horizontal: 16, vertical: 8),
+        child: Column(
+          mainAxisSize: MainAxisSize.min,
+          children: [
+            Row(
+              mainAxisSize: MainAxisSize.min,
+              mainAxisAlignment: MainAxisAlignment.center,
+              children: [
+                AssetIcon(
+                  path: 'assets/images/icons/unreal_u_logo.svg',
+                  size: 24,
+                ),
+                SizedBox(width: 20),
+                Text(
+                  AppLocalizations.of(context)!.connectScreenRootActorDialogTitle,
+                  style: Theme.of(context).textTheme.displayLarge,
+                ),
+              ],
+            ),
+            SizedBox(height: 32),
+            Text(
+              AppLocalizations.of(context)!.connectScreenRootActorDialogMessage,
+            ),
+            SizedBox(height: 32),
+            ..._actorManager.getActorsOfClass(nDisplayRootActorClassName).map(
+                  (e) => NDisplayItem(
+                    e,
+                    bSelected: _selectedActor == e,
+                    onChanged: (value) => setState(() => _selectedActor = value),
+                  ),
+                ),
+            if (!widget.bIsAlreadyConnected || _selectedActor != null)
+              Padding(
+                padding: EdgeInsets.only(top: 20),
+                child: Row(
+                  mainAxisAlignment: MainAxisAlignment.end,
+                  children: [
+                    if (!widget.bIsAlreadyConnected)
+                      EpicLozengeButton(
+                        onPressed: () => Navigator.pop(context),
+                        label: AppLocalizations.of(context)!.menuButtonCancel,
+                        color: (_selectedActor != null) ? Colors.transparent : Theme.of(context).colorScheme.secondary,
+                      ),
+                    if (_selectedActor != null)
+                      Padding(
+                        padding: EdgeInsets.only(left: 10),
+                        child: EpicLozengeButton(
+                            onPressed: () => Navigator.of(context).pop(_selectedActor),
+                            label: widget.bIsAlreadyConnected
+                                ? AppLocalizations.of(context)!.menuButtonProceed
+                                : AppLocalizations.of(context)!.nDisplaySelectorModalConnectButtonLabel),
+                      ),
+                  ],
+                ),
+              ),
+          ],
+        ),
+      ),
+    );
+  }
+
+  /// Called when the list of root actors changes
+  void _onRootActorUpdate(ActorUpdateDetails details) {
+    if (!mounted) {
+      return;
+    }
+
+    final Set<UnrealObject> rootActors = _actorManager.getActorsOfClass(nDisplayRootActorClassName);
+    if (rootActors.length == 1) {
+      // Only one root actor, so use that one
+      Navigator.pop(context, rootActors.first);
+      return;
+    }
+
+    if (rootActors.length == 0) {
+      // No root actors, so nothing to show anymore
+      Navigator.pop(context);
+      return;
+    }
+
+    // Refresh to show new list of actors
+    setState(() {});
+  }
+}
+
+///Visual representation/Item for available root actors to be rendered on [NDisplaySelectorDialog].
+class NDisplayItem extends StatelessWidget {
+  const NDisplayItem(this.nDisplay, {this.bSelected = false, required this.onChanged, Key? key}) : super(key: key);
+
+  /// NDisplay actor being rendered for selection.
+  final UnrealObject nDisplay;
+
+  /// Whether the currently represented actor with name value of [nDisplay.name] is currently selected by the user.
+  final bool bSelected;
+
+  /// Value change callback to notify parent widget of currently selected actor, returns [nDisplay].
+  final ValueChanged<UnrealObject> onChanged;
+
+  @override
+  Widget build(BuildContext context) {
+    return InkWell(
+      onTap: () => onChanged(nDisplay),
+      child: Container(
+        decoration: BoxDecoration(
+          border: Border.all(width: 3, color: bSelected ? Theme.of(context).colorScheme.primary : UnrealColors.gray42),
+          borderRadius: BorderRadius.circular(8),
+          color: bSelected ? Theme.of(context).colorScheme.primary : null,
+        ),
+        padding: EdgeInsets.all(10),
+        margin: EdgeInsets.symmetric(horizontal: 30, vertical: 5),
+        child: Row(
+          children: [
+            AssetIcon(
+              path: 'assets/images/icons/ndisplay.svg',
+              size: 32,
+            ),
+            SizedBox(width: 16),
+            Text(
+              nDisplay.name,
+              style: Theme.of(context)
+                  .textTheme
+                  .titleLarge!
+                  .copyWith(color: bSelected ? Theme.of(context).colorScheme.onPrimary : UnrealColors.gray56),
+            ),
+          ],
+        ),
+      ),
+    );
+  }
+}
diff --git a/Engine/Extras/VirtualProduction/EpicStageApp/lib/widgets/screens/connect/views/passphrase_modal.dart b/Engine/Extras/VirtualProduction/EpicStageApp/lib/widgets/screens/connect/views/passphrase_modal.dart
new file mode 100644
--- /dev/null
+++ b/Engine/Extras/VirtualProduction/EpicStageApp/lib/widgets/screens/connect/views/passphrase_modal.dart
@@ -0,0 +1,196 @@
+// Copyright Epic Games, Inc. All Rights Reserved.
+
+import 'package:flutter/material.dart';
+import 'package:flutter_gen/gen_l10n/app_localizations.dart';
+import 'package:logging/logging.dart';
+import 'package:provider/provider.dart';
+
+import '../../../../models/engine_connection.dart';
+import '../../../../models/engine_passphrase_manager.dart';
+import '../../../../utilities/net_utilities.dart';
+import '../../../../utilities/unreal_colors.dart';
+import '../../../elements/epic_checkbox.dart';
+import '../../../elements/epic_icon_button.dart';
+import '../../../elements/modal.dart';
+import '../mixins/connect_mixin.dart';
+
+final _log = Logger('PassphraseModal');
+
+/// A modal dialog that lets the user set a passphrase to use for the last attempted connection, then reconnects.
+class PassphraseModalDialog extends StatefulWidget {
+  const PassphraseModalDialog({Key? key, this.initialErrorMessage}) : super(key: key);
+
+  /// The initial error message to show to the user, if any.
+  final String? initialErrorMessage;
+
+  @override
+  State<StatefulWidget> createState() => _PassphraseModalDialogState();
+}
+
+class _PassphraseModalDialogState extends State<PassphraseModalDialog> with ConnectMixin {
+  /// Text controller for the input field displayed to the user.
+  final _textController = TextEditingController();
+
+  /// The connection for which we're modifying the password.
+  late final ConnectionData _connectionData;
+
+  /// False if this is still retrieving connection and passphrase data.
+  bool _bIsReady = false;
+
+  /// Whether to show the passphrase to the user.
+  bool _bShowPassphrase = false;
+
+  /// Error message to show to the user, if any.
+  String? _errorMessage;
+
+  @override
+  void initState() {
+    super.initState();
+
+    _errorMessage = widget.initialErrorMessage;
+    _initAsync();
+  }
+
+  @override
+  Widget build(BuildContext context) {
+    const double buttonWidth = 110;
+    final localizations = AppLocalizations.of(context)!;
+
+    return ModalDialogCard(
+      child: SizedBox(
+        width: 300,
+        child: IntrinsicHeight(
+          child: _bIsReady
+              ? Column(
+                  children: [
+                    ModalDialogTitle(title: localizations.connectScreenPasshphraseModalTitle),
+                    if (_errorMessage != null)
+                      ModalDialogSection(
+                        child: Text(
+                          _errorMessage!,
+                          style: Theme.of(context).textTheme.bodyMedium!.copyWith(
+                                color: UnrealColors.highlightRed,
+                              ),
+                          textAlign: TextAlign.start,
+                        ),
+                      ),
+                    ModalDialogSection(
+                      child: Column(
+                        children: [
+                          Row(
+                            children: [
+                              Expanded(
+                                child: SizedBox(
+                                  height: 36,
+                                  child: TextField(
+                                    autofocus: true,
+                                    maxLines: 1,
+                                    cursorWidth: 1,
+                                    controller: _textController,
+                                    keyboardAppearance: Brightness.dark,
+                                    style: Theme.of(context).textTheme.bodyMedium,
+                                    keyboardType: TextInputType.visiblePassword,
+                                    onEditingComplete: _submit,
+                                    obscureText: !_bShowPassphrase,
+                                  ),
+                                ),
+                              ),
+                            ],
+                          ),
+                          SizedBox(height: 12),
+                          EpicCheckbox(
+                            bChecked: _bShowPassphrase,
+                            label: localizations.showPassphraseToggleLabel,
+                            onPressed: () => setState(() {
+                              _bShowPassphrase = !_bShowPassphrase;
+                            }),
+                          ),
+                        ],
+                      ),
+                    ),
+                    ModalDialogSection(
+                      child: Row(
+                        mainAxisAlignment: MainAxisAlignment.spaceAround,
+                        children: [
+                          EpicLozengeButton(
+                            label: localizations.menuButtonCancel,
+                            width: buttonWidth,
+                            color: Colors.transparent,
+                            onPressed: _cancel,
+                          ),
+                          EpicLozengeButton(
+                            label: localizations.menuButtonSubmit,
+                            width: buttonWidth,
+                            onPressed: _submit,
+                          ),
+                        ],
+                      ),
+                    ),
+                  ],
+                )
+              : ModalDialogSection(
+                  child: Center(
+                    child: SizedBox.square(
+                      child: CircularProgressIndicator(),
+                      dimension: 80,
+                    ),
+                  ),
+                ),
+        ),
+      ),
+    );
+  }
+
+  /// Retrieve the connection data and current password, then fill the text field with it.
+  void _initAsync() async {
+    final connectionManager = Provider.of<EngineConnectionManager>(context, listen: false);
+
+    final ConnectionData? connectionData = await connectionManager.getLastConnectionData();
+    if (connectionData == null) {
+      _log.warning('Tried to set password but no last connection existed');
+      Navigator.pop(context);
+      return;
+    }
+
+    _connectionData = connectionData;
+
+    final passphraseManager = Provider.of<EnginePassphraseManager>(context, listen: false);
+    final String? passphrase = await passphraseManager.getPassphrase(connectionData);
+
+    if (passphrase != null) {
+      _textController.text = passphrase;
+      _textController.selection = TextSelection(baseOffset: 0, extentOffset: _textController.text.length);
+    }
+
+    setState(() {
+      _bIsReady = true;
+    });
+  }
+
+  /// Save the new passphrase, then reconnect to the engine.
+  void _submit() async {
+    if (!_bIsReady) {
+      return;
+    }
+
+    final passphraseManager = Provider.of<EnginePassphraseManager>(context, listen: false);
+    passphraseManager.setPassphrase(_connectionData, _textController.text);
+
+    final EngineConnectionResult result = await connect(
+      _connectionData,
+      // The user will just return to this modal, so we don't need to create another
+      bShowPassphraseModalOnReject: false,
+    );
+
+    if (result == EngineConnectionResult.passphraseRejected) {
+      setState(() {
+        _errorMessage = AppLocalizations.of(context)!.connectScreenPassphraseIncorrectErrorMessage;
+      });
+    }
+  }
+
+  /// Cancel the interaction and pop the modal.
+  void _cancel() {
+    Navigator.pop(context);
+  }
+}
diff --git a/Engine/Extras/VirtualProduction/EpicStageApp/lib/widgets/screens/connect/views/quick_action_grid.dart b/Engine/Extras/VirtualProduction/EpicStageApp/lib/widgets/screens/connect/views/quick_action_grid.dart
new file mode 100644
--- /dev/null
+++ b/Engine/Extras/VirtualProduction/EpicStageApp/lib/widgets/screens/connect/views/quick_action_grid.dart
@@ -0,0 +1,34 @@
+// Copyright Epic Games, Inc. All Rights Reserved.
+
+import 'package:flutter/material.dart';
+
+import '../../../../utilities/net_utilities.dart';
+import 'connection_tile.dart';
+import 'manual_connect_tile.dart';
+
+/// Quick action column showing most recent connection instances of UE.
+class QuickActionGrid extends StatelessWidget {
+  const QuickActionGrid({Key? key, required this.recentConnections}) : super(key: key);
+
+  ///List of most recent/saved instances of UE connection.
+  final List<ConnectionData> recentConnections;
+
+  @override
+  Widget build(BuildContext context) {
+    return Padding(
+      padding: EdgeInsets.all(32),
+      child: Column(
+        children: [
+          ...recentConnections.map((e) => Padding(
+                padding: EdgeInsets.only(bottom: 32),
+                child: ConnectionTile(
+                  connectionData: e,
+                  bPrimary: true,
+                ),
+              )),
+          ManualConnectTile(),
+        ],
+      ),
+    );
+  }
+}
diff --git a/Engine/Extras/VirtualProduction/EpicStageApp/lib/widgets/screens/eula/eula_screen.dart b/Engine/Extras/VirtualProduction/EpicStageApp/lib/widgets/screens/eula/eula_screen.dart
new file mode 100644
--- /dev/null
+++ b/Engine/Extras/VirtualProduction/EpicStageApp/lib/widgets/screens/eula/eula_screen.dart
@@ -0,0 +1,408 @@
+// Copyright Epic Games, Inc. All Rights Reserved.
+
+import '../../elements/epic_scroll_view.dart';
+import 'package:flutter/cupertino.dart';
+import 'package:provider/provider.dart';
+import 'package:flutter/material.dart';
+import 'package:flutter_gen/gen_l10n/app_localizations.dart';
+
+import '../../../utilities/preferences_bundle.dart';
+import '../../../utilities/unreal_colors.dart';
+import '../../elements/parsed_rich_text.dart';
+import '../connect/connect.dart';
+
+/// Data types for EULA content.
+enum _EulaDataType { paragraph, header, sub }
+
+/// Screen for rendering and displaying content of the EULA doc.
+class EulaScreen extends StatefulWidget {
+  const EulaScreen({Key? key}) : super(key: key);
+
+  static const String route = '/eula';
+
+  @override
+  State<EulaScreen> createState() => _EulaScreenState();
+}
+
+class _EulaScreenState extends State<EulaScreen> {
+  /// Persistent store state reference.
+  late PreferencesBundle store;
+
+  /// whether the user agrees and accepts the EULA agreement or not.
+  late bool bHasAcceptedEula;
+
+  @override
+  void initState() {
+    super.initState();
+    store = Provider.of<PreferencesBundle>(context, listen: false);
+    bHasAcceptedEula = store.persistent.getBool('common.bHasAcceptedEula', defaultValue: false).getValue();
+  }
+
+  @override
+  Widget build(BuildContext context) {
+    return Scaffold(
+      appBar: AppBar(
+        title: Text(
+          AppLocalizations.of(context)!.eulaAppBarTitle,
+          style: Theme.of(context).textTheme.displayLarge,
+        ),
+      ),
+      body: EpicScrollView(
+        child: Column(
+          children: [
+            Padding(
+              padding: const EdgeInsets.symmetric(vertical: 20.0),
+              child: Text(
+                'Unreal Stage End User License Agreement',
+                style: Theme.of(context).textTheme.displayLarge,
+              ),
+            ),
+            ListView.builder(
+              physics: NeverScrollableScrollPhysics(),
+              itemCount: EulaDTO.data.length,
+              shrinkWrap: true,
+              itemBuilder: (BuildContext context, int index) {
+                return _EulaItem(data: EulaDTO.data[index]);
+              },
+            ),
+          ],
+        ),
+      ),
+      bottomNavigationBar: Container(
+        padding: EdgeInsets.symmetric(vertical: 18, horizontal: 18),
+        decoration: BoxDecoration(color: UnrealColors.gray06),
+        child: CupertinoButton(
+          child: Text(
+            bHasAcceptedEula
+                ? AppLocalizations.of(context)!.menuButtonOK
+                : AppLocalizations.of(context)!.eulaAcceptButtonLabel,
+          ),
+          onPressed: () {
+            if (bHasAcceptedEula) {
+              if (Navigator.of(context).canPop()) {
+                Navigator.pop(context);
+              }
+            } else {
+              store.persistent.setBool('common.bHasAcceptedEula', true);
+              Navigator.pushReplacementNamed(context, ConnectScreen.route);
+            }
+          },
+        ),
+      ),
+    );
+  }
+}
+
+/// A widget for visually representing heading and paragraphs on the EULA docs.
+class _EulaItem extends StatelessWidget {
+  const _EulaItem({required this.data, Key? key}) : super(key: key);
+
+  /// Data objects for diff paragraphs in the EULA doc.
+  final EulaDTO data;
+
+  @override
+  Widget build(BuildContext context) {
+    return Container(
+      padding: EdgeInsets.symmetric(horizontal: 16, vertical: 8),
+      child: Column(
+        crossAxisAlignment: CrossAxisAlignment.start,
+        children: [
+          data.type == _EulaDataType.header
+              ? Text(data.text, style: Theme.of(context).textTheme.displayLarge)
+              : ParsedRichText(data.text),
+          ...(data.subs ?? []).map<Widget>((e) => buildSubText(e, context)).toList(),
+        ],
+      ),
+    );
+  }
+
+  /// Builds sub texts under some paragraphs using info from [data] under the current [context].
+  Widget buildSubText(EulaDTO data, BuildContext context) {
+    return Container(
+      padding: EdgeInsets.only(left: 20).add(EdgeInsets.symmetric(vertical: 8)),
+      child: Row(
+        crossAxisAlignment: CrossAxisAlignment.start,
+        children: [
+          Text(data.numbering!, style: Theme.of(context).textTheme.displayMedium),
+          SizedBox(width: 10),
+          Expanded(child: Text(data.text, style: Theme.of(context).textTheme.displayMedium)),
+        ],
+      ),
+    );
+  }
+}
+
+/// EULA data transfer object.
+class EulaDTO {
+  /// if the text being processed is a header, sub text or paragraph.
+  _EulaDataType? type = _EulaDataType.header;
+
+  /// Text to be rendered.
+  String text;
+
+  /// Sub text under a paragraph [text] when [type] = [_EulaDataType.sub].
+  List<EulaDTO>? subs;
+
+  /// optionally needed for numbered sub text [subs].
+  String? numbering;
+
+  EulaDTO({this.type, required this.text, this.subs, this.numbering});
+
+  /// Data source for EULA docs.
+  static final data = <EulaDTO>[
+    EulaDTO(
+      text:
+          'Please read this Agreement carefully.  It is a legal document that explains your rights and obligations related to your use of the Software, including any Services you access through the Software.  By downloading or using the Software, or by otherwise indicating your acceptance of this Agreement, you are agreeing to be bound by the terms of this Agreement.  If you do not or cannot agree to the terms of this Agreement, please do not download or use this Software.',
+      type: _EulaDataType.paragraph,
+    ),
+    EulaDTO(
+      text:
+          'Your use of the Services is also governed by Epic’s Terms of Service, which may be found at https://www.epicgames.com/tos and Epic’s Privacy Policy, which may be found at https://www.epicgames.com/privacypolicy. By downloading or using the Software, you also agree to Epic’s Terms of Service and Privacy Policy.',
+      type: _EulaDataType.paragraph,
+    ),
+    EulaDTO(
+      text:
+          'Certain words or phrases are defined to have certain meanings when used in this Agreement.  Those words and phrases are defined below in Section 14.',
+      type: _EulaDataType.paragraph,
+    ),
+    EulaDTO(
+      text:
+          'If your primary residence is in the United States of America, your agreement is with Epic Games, Inc.  If it is not in the United States of America, your agreement is with Epic Games Commerce GmbH.',
+      type: _EulaDataType.paragraph,
+    ),
+    EulaDTO(
+      text: '1. License Grant',
+      type: _EulaDataType.header,
+    ),
+    EulaDTO(
+      text:
+          'Epic grants you a personal, non-exclusive, non-transferable, non-sublicensable limited right and license to use and install the Software on an Authorized Device (the <b>“License”</b>).  The rights that Epic grants you under the License are subject to the terms of this Agreement, and you may only make use of the License if you comply with all applicable terms.',
+      type: _EulaDataType.paragraph,
+    ),
+    EulaDTO(
+      text:
+          'The License becomes effective on the date you accept this Agreement. The Software is licensed, not sold, to you under the License. The License does not grant you any title or ownership in the Software. ',
+      type: _EulaDataType.paragraph,
+    ),
+    EulaDTO(
+      text: '2. License Conditions',
+      type: _EulaDataType.header,
+    ),
+    EulaDTO(
+      text:
+          'You may not do any of the following with respect to the Software or any of its parts:  (a) copy, reproduce, distribute, display, or use it in a way that is not expressly authorized in this Agreement; (b) sell, rent, lease, license, distribute, or otherwise transfer it; (c) reverse engineer, derive source code from, modify, adapt, translate, decompile, or disassemble it or make derivative works based on it; (d) remove, disable, circumvent, or modify any proprietary notice or label or security technology included in it; (e) create, develop, distribute, or use any unauthorized software programs; (f) use it to infringe or violate the rights of any third party, including but not limited to any intellectual property, publicity, or privacy rights; (g) use, export, or re-export it in violation of any applicable law or regulation; (h) use it on a device other than an Authorized Device; or (i) behave in a manner which is detrimental to the enjoyment of the Software or any Apps by other users as intended by Epic, in Epic’s sole judgment, including but not limited to the following – harassment, use of abusive or offensive language, game abandonment, game sabotage, spamming, social engineering, or scamming.',
+      type: _EulaDataType.paragraph,
+    ),
+    EulaDTO(
+      text: '3. Updates and Patches',
+      type: _EulaDataType.header,
+    ),
+    EulaDTO(
+      text:
+          'Epic may provide patches, updates, or upgrades to the Software that must be installed in order for you to continue to use the Software or Services.  Epic may update the Software remotely without notifying you, and you hereby consent to Epic applying patches, updates, and upgrades.  Epic may modify, suspend, discontinue, substitute, replace, or limit your access to any aspect of the Software or Services at any time.  You acknowledge that your use of the Software or Services does not confer on you any interest, monetary or otherwise, in any aspect or feature of the Software or any Apps.  Epic does not have any maintenance or support obligations with respect to the Software or Services.',
+      type: _EulaDataType.paragraph,
+    ),
+    EulaDTO(
+      text: '4. Feedback',
+      type: _EulaDataType.header,
+    ),
+    EulaDTO(
+      text:
+          'If you provide Epic with any Feedback, you hereby grant Epic a non-exclusive, fully-paid, royalty-free, irrevocable, perpetual, transferable, sublicensable license to reproduce, distribute, modify, prepare derivative works based on, publicly perform, publicly display, make, have made, use, sell, offer to sell, import, and otherwise exploit that Feedback for any purposes, for all current and future methods and forms of exploitation in any country.  If any such rights may not be licensed under applicable law (such as moral and other personal rights), you hereby waive and agree not to assert all such rights.  You understand and agree that Epic is not required to make any use of any Feedback that you provide.  You agree that if Epic makes use of your Feedback, Epic is not required to credit or compensate you for your contribution. You represent and warrant that you have sufficient rights in any Feedback that you provide to Epic to grant Epic and other affected parties the rights described above.  This includes but is not limited to intellectual property rights and other proprietary or personal rights.',
+      type: _EulaDataType.paragraph,
+    ),
+    EulaDTO(
+      text: '5. Ownership/Third Party Licenses',
+      type: _EulaDataType.header,
+    ),
+    EulaDTO(
+      text:
+          'Epic and its licensors own all title, ownership rights, and intellectual property rights in the Software and Services.  Epic, Epic Games, Unreal, Unreal Engine, and their respective logos, are trademarks or registered trademarks of Epic and its affiliates in the United States of America and elsewhere.  All rights granted to you under this Agreement are granted by express license only and not by sale.  No license or other rights shall be created hereunder by implication, estoppel, or otherwise.',
+      type: _EulaDataType.paragraph,
+    ),
+    EulaDTO(
+      text:
+          'The Software includes certain components provided by Epic’s licensors.  A list of credits and notices for third party components may be found in the Software interface.',
+      type: _EulaDataType.paragraph,
+    ),
+    EulaDTO(
+      text: '6. Disclaimers and Limitation of Liability',
+      type: _EulaDataType.header,
+    ),
+    EulaDTO(
+      text:
+          '<b>The Software and Services is provided on an “as is” and “as available” basis, “with all faults” and without warranty of any kind.  Epic, its licensors, and its and their affiliates disclaim all warranties, conditions, common law duties, and representations (express, implied, oral, and written) with respect to the Software and Services, including without limitation all express, implied, and statutory warranties and conditions of any kind, such as title, non-interference with your enjoyment, authority, non-infringement, merchantability, fitness or suitability for any purpose (whether or not Epic knows or has reason to know of any such purpose), system integration, accuracy or completeness, results, reasonable care, workmanlike effort, lack of negligence, and lack of viruses, whether alleged to arise under law, by reason of custom or usage in the trade, or by course of dealing.  Without limiting the generality of the foregoing, Epic, its licensors, and its and their affiliates make no warranty that (1) the Software or Services will operate properly, (2) that the Software or Services will meet your requirements, (3) that the operation of the Software or Services will be uninterrupted, bug free, or error free in any or all circumstances, or (4) that any defects in the Software or Services can or will be corrected.  Any warranty against infringement that may be provided in Section 2-312 of the Uniform Commercial Code or in any other comparable statute is expressly disclaimed.  Epic, its licensors, and its and their affiliates do not guarantee continuous, error-free, virus-free, or secure operation of or access to the Software or Services.  This paragraph will apply to the maximum extent permitted by applicable law.<\/b>',
+      type: _EulaDataType.paragraph,
+    ),
+    EulaDTO(
+      text:
+          '<b>To the maximum extent permitted by applicable law, neither Epic, nor its licensors, nor its or their affiliates, nor any of Epic’s service providers (collectively, the “Epic Parties”), shall be liable in any way for any loss of profits or any indirect, incidental, consequential, special, punitive, or exemplary damages, arising out of or in connection with this Agreement or the Software or Services, or the delay or inability to use or lack of functionality of the Software or Services, even in the event of an Epic Party’s fault, tort (including negligence), strict liability, indemnity, product liability, breach of contract, breach of warranty, or otherwise and even if an Epic Party has been advised of the possibility of such damages.  Further, to the maximum extent permitted by applicable law, the aggregate liability of the Epic Parties arising out of or in connection with this Agreement or the Software  or Services will not exceed the greater of (a) the total amounts you have paid (if any) to Epic for the Software or within the Software during the twelve (12) months immediately preceding the events giving rise to such liability and (b) U.S. \$5.00. These limitations and exclusions regarding damages apply even if any remedy fails to provide adequate compensation.<\/b>',
+      type: _EulaDataType.paragraph,
+    ),
+    EulaDTO(
+      text:
+          '<b>Notwithstanding the foregoing, some countries, states, provinces or other jurisdictions do not allow the exclusion of certain warranties or the limitation of liability as stated above, so the above terms may not apply to you. Instead, in such jurisdictions, the foregoing exclusions and limitations shall apply only to the extent permitted by the laws of such jurisdictions. Also, you may have additional legal rights in your jurisdiction, and nothing in this Agreement will prejudice the statutory rights that you may have as a consumer of the Software or Services.<\/b>',
+      type: _EulaDataType.paragraph,
+    ),
+    EulaDTO(
+      text: '7. Indemnity',
+      type: _EulaDataType.header,
+    ),
+    EulaDTO(
+      text:
+          'You agree to indemnify, pay the defense costs of, and hold Epic, its licensors, its and their affiliates, and its and their employees, officers, directors, agents, contractors, and other representatives harmless from all claims, demands, actions, losses, liabilities, and expenses (including attorneys’ fees, costs, and expert witnesses’ fees) that arise from or in connection with (a) any claim that, if true, would constitute a breach by you of this Agreement or negligence by you, or (b) any act or omission by you in using the Software or Services.  You agree to reimburse Epic on demand for any defense costs incurred by Epic and any payments made or loss suffered by Epic, whether in a court judgment or settlement, based on any matter covered by this Section 7.',
+      type: _EulaDataType.paragraph,
+    ),
+    EulaDTO(
+      text:
+          'If you are prohibited by law from entering into the indemnification obligation above, then you assume, to the extent permitted by law, all liability for all claims, demands, actions, losses, liabilities, and expenses (including attorneys’ fees, costs and expert witnesses’ fees) that are the stated subject matter of the indemnification obligation above.',
+      type: _EulaDataType.paragraph,
+    ),
+    EulaDTO(
+      text: '8. Termination',
+      type: _EulaDataType.header,
+    ),
+    EulaDTO(
+      text:
+          'Without limiting any other rights of Epic, this Agreement will terminate automatically without notice if you fail to comply with any of its terms and conditions.  You may also terminate this Agreement by deleting the Software.  Upon any termination, the License will automatically terminate, you may no longer exercise any of the rights granted to you by the License, and you must destroy all copies of the Software in your possession.',
+      type: _EulaDataType.paragraph,
+    ),
+    EulaDTO(
+      text:
+          '<b>Except to the extent required by law, all payments and fees are non-refundable under all circumstances, regardless of whether or not this Agreement has been terminated.<\/b>',
+      type: _EulaDataType.paragraph,
+    ),
+    EulaDTO(
+      text: 'Sections 2, 6-11, and 13-15 will survive any termination of this Agreement.',
+      type: _EulaDataType.paragraph,
+    ),
+    EulaDTO(
+      text: '9. Governing Law and Jurisdiction ',
+      type: _EulaDataType.header,
+    ),
+    EulaDTO(
+      text:
+          'You agree that this Agreement will be deemed to have been made and executed in the State of North Carolina, U.S.A., and any dispute will be resolved in accordance with the laws of North Carolina, excluding that body of law related to choice of laws, and of the United States of America.  Any action or proceeding brought to enforce the terms of this Agreement or to adjudicate any dispute must be brought in the Superior Court of Wake County, State of North Carolina or the United States District Court for the Eastern District of North Carolina.  You agree to the exclusive jurisdiction and venue of these courts.  You waive any claim of inconvenient forum and any right to a jury trial.  The Convention on Contracts for the International Sale of Goods will not apply.  Any law or regulation which provides that the language of a contract shall be construed against the drafter will not apply to this Agreement.',
+      type: _EulaDataType.paragraph,
+    ),
+    EulaDTO(
+      text: '10. Class Action Waiver',
+      type: _EulaDataType.header,
+    ),
+    EulaDTO(
+      text:
+          '<b>You agree not to bring or participate in a class or representative action, private attorney general action, or collective arbitration related to the Software or Services or this Agreement.  You also agree not to seek to combine any action or arbitration related to the Software or Services or this Agreement with any other action or arbitration without the consent of all parties to this Agreement and all other actions or arbitrations.<\/b>',
+      type: _EulaDataType.paragraph,
+    ),
+    EulaDTO(
+      text: '11. U.S. Government Matters',
+      type: _EulaDataType.header,
+    ),
+    EulaDTO(
+      text:
+          'The Software is a “Commercial Item” (as defined at 48 C.F.R. §2.101), consisting of “Commercial Computer Software” and “Commercial Computer Software Documentation” (as used in 48 C.F.R. §12.212 or 48 C.F.R. §227.7202, as applicable).  The Software is being licensed to U.S. Government end users only as Commercial Items and with only those rights as are granted to other licensees under this Agreement.',
+      type: _EulaDataType.paragraph,
+    ),
+    EulaDTO(
+      text:
+          'You represent and warrant to Epic that you are not located in a country that is subject to a U.S. Government embargo or that has been designated by the U.S. Government as a “terrorist supporting” country, and that you are not listed on any U.S. Government list of prohibited or restricted parties.',
+      type: _EulaDataType.paragraph,
+    ),
+    EulaDTO(
+      text: '12. Amendments of this Agreement ',
+      type: _EulaDataType.header,
+    ),
+    EulaDTO(
+      text:
+          '<b>Epic may issue an amended Agreement, Terms of Service, or Privacy Policy at any time in its discretion by posting the amended Agreement, Terms of Service, or Privacy Policy on its website or by providing you with digital access to amended versions of any of these documents when you next access the Software.  If any amendment to this Agreement, the Terms of Service, or the Privacy Policy is not acceptable to you, you may terminate this Agreement and must stop using the Software.  Your continued use of the Software will demonstrate your acceptance of the amended Agreement, Terms of Service, and Privacy Policy.<\/b>',
+      type: _EulaDataType.paragraph,
+    ),
+    EulaDTO(
+      text: '13. No Assignment',
+      type: _EulaDataType.header,
+    ),
+    EulaDTO(
+      text:
+          'You may not, without the prior written consent of Epic, assign, transfer, charge, or sub-contract all or any of your rights or obligations under this Agreement, and any attempt without that consent will be null and void.  If restrictions on transfer of the Software in this Agreement are not enforceable under the law of your country, then this Agreement will be binding on any recipient of the Software.  Epic may at any time assign, transfer, charge, or sub-contract all or any of its rights or obligations under this Agreement.',
+      type: _EulaDataType.paragraph,
+    ),
+    EulaDTO(
+      text: '14. Definitions',
+      type: _EulaDataType.header,
+    ),
+    EulaDTO(
+      text: 'As used in this Agreement, the following capitalized words have the following meanings: ',
+      type: _EulaDataType.paragraph,
+    ),
+    EulaDTO(
+      text:
+          '<b>“Authorized Device”<\/b> means a device on which the Software is designed and designated by Epic to be operated. An Authorized Device does not include a virtual machine or other emulation of such a mobile device.',
+      type: _EulaDataType.paragraph,
+    ),
+    EulaDTO(
+      text: '<b>“Epic”<\/b> means, depending on the location of your primary residence:',
+      type: _EulaDataType.paragraph,
+      subs: [
+        EulaDTO(
+            text:
+                'Epic Games, Inc., a Maryland Corporation having its principal business offices at Box 254, 2474 Walnut Street, Cary, North Carolina, 27518, U.S.A.; or',
+            type: _EulaDataType.sub,
+            numbering: 'a.'),
+        EulaDTO(
+          text:
+              'Epic Games Epic Games Commerce, GmbH, a limited liability company organized under the laws of Switzerland, having its principal business offices at Platz 10, 6039 Root D4, Switzerland.  ',
+          type: _EulaDataType.sub,
+          numbering: 'b.',
+        ),
+      ],
+    ),
+    EulaDTO(
+      text:
+          '<b>“Feedback”<\/b> means any feedback or suggestions that you provide to Epic regarding the Software, Services, or other Epic products and services.',
+      type: _EulaDataType.paragraph,
+    ),
+    EulaDTO(
+      text:
+          '<b>“Services”<\/b> means any services made available to you through the Software, including services to manage your Epic account, connect to third party services.',
+      type: _EulaDataType.paragraph,
+    ),
+    EulaDTO(
+      text: '<b>“Software”<\/b> means Epic’s proprietary software known as Unreal Stage.',
+      type: _EulaDataType.paragraph,
+    ),
+    EulaDTO(
+      text: '15. Miscellaneous',
+      type: _EulaDataType.header,
+    ),
+    EulaDTO(
+      text:
+          'This Agreement and any document or information referred to in this Agreement constitute the entire agreement between you and Epic relating to the subject matter covered by this Agreement.  All other communications, proposals, and representations with respect to the subject matter covered by this Agreement are excluded.',
+      type: _EulaDataType.paragraph,
+    ),
+    EulaDTO(
+      text:
+          'The original of this Agreement is in English; any translations are provided for reference purposes only.  You waive any right you may have under the law of your country to have this Agreement written or construed in the language of any other country.',
+      type: _EulaDataType.paragraph,
+    ),
+    EulaDTO(
+      text:
+          'This Agreement describes certain legal rights. You may have other rights under the laws of your jurisdiction. This Agreement does not change your rights under the laws of your jurisdiction if the laws of your jurisdiction do not permit it to do so.  Limitations and exclusions of warranties and remedies in this Agreement may not apply to you because your jurisdiction may not allow them in your particular circumstance.  In the event that certain provisions of this Agreement are held by a court or tribunal of competent jurisdiction to be unenforceable, those provisions shall be enforced only to the furthest extent possible under applicable law and the remaining terms of this Agreement will remain in full force and effect.',
+      type: _EulaDataType.paragraph,
+    ),
+    EulaDTO(
+      text:
+          'Any act by Epic to exercise, or failure or delay in exercise of, any of its rights under this Agreement, at law or in equity will not be deemed a waiver of those or any other rights or remedies available in contract, at law or in equity.',
+      type: _EulaDataType.paragraph,
+    ),
+    EulaDTO(
+      text:
+          'You agree that this Agreement does not confer any rights or remedies on any person other than the parties to this Agreement, except as expressly stated.',
+      type: _EulaDataType.paragraph,
+    ),
+    EulaDTO(
+      text:
+          'Epic’s obligations are subject to existing laws and legal process, and Epic may comply with law enforcement or regulatory requests or requirements despite any contrary term in this Agreement.',
+      type: _EulaDataType.paragraph,
+    ),
+  ];
+}
