Branch snapshot for CL 27138787

diff --git a/Engine/Extras/VirtualProduction/EpicStageApp/lib/widgets/screens/main/sidebar/outliner_filter_menu.dart b/Engine/Extras/VirtualProduction/EpicStageApp/lib/widgets/screens/main/sidebar/outliner_filter_menu.dart
new file mode 100644
--- /dev/null
+++ b/Engine/Extras/VirtualProduction/EpicStageApp/lib/widgets/screens/main/sidebar/outliner_filter_menu.dart
@@ -0,0 +1,91 @@
+// Copyright Epic Games, Inc. All Rights Reserved.
+
+import 'package:flutter/material.dart';
+import 'package:flutter_gen/gen_l10n/app_localizations.dart';
+import 'package:provider/provider.dart';
+import 'package:streaming_shared_preferences/streaming_shared_preferences.dart';
+
+import '../../../../models/settings/outliner_panel_settings.dart';
+import '../../../../models/unreal_object_filters.dart';
+import '../../../elements/dropdown_button.dart';
+import '../../../elements/dropdown_list_menu.dart';
+import '../../../elements/layout/card.dart';
+import '../../../elements/list_menu.dart';
+
+/// Modal dropdown button that displays outliner filter settings.
+class OutlinerFilterButton extends StatelessWidget {
+  const OutlinerFilterButton({Key? key}) : super(key: key);
+
+  @override
+  Widget build(BuildContext context) {
+    final settings = Provider.of<OutlinerPanelSettings>(context, listen: false);
+
+    return ModalDropdownButton(
+      buttonBuilder: (context, state) => Material(
+        color: Colors.transparent,
+        child: PreferenceBuilder(
+          preference: settings.selectedFilters,
+          builder: (context, Set<UnrealObjectFilter> filters) => CardSubHeaderButton(
+            iconPath: filters.isEmpty ? 'assets/images/icons/filter.svg' : 'assets/images/icons/filter_on.svg',
+            tooltipMessage: AppLocalizations.of(context)!.outlinerFiltersButtonTooltip,
+            bIsToggledOn: state != ModalDropdownButtonState.closed,
+            bIsVisualOnly: true,
+          ),
+        ),
+      ),
+      menuBuilder: (menuContext, originTabBuilder) => _OutlinerFilterDropDownMenu(
+        originTabBuilder: originTabBuilder,
+        settings: settings,
+      ),
+    );
+  }
+}
+
+/// Dropdown menu shown to select which filters to apply to the Outliner panel.
+class _OutlinerFilterDropDownMenu extends StatelessWidget {
+  const _OutlinerFilterDropDownMenu({
+    Key? key,
+    required this.settings,
+    this.originTabBuilder,
+  }) : super(key: key);
+
+  /// The [OutlinerPanelSettings] from the originating context.
+  final OutlinerPanelSettings settings;
+
+  /// Function that builds the origin tab for this menu.
+  final Widget Function(BuildContext, bool)? originTabBuilder;
+
+  @override
+  Widget build(BuildContext context) {
+    return PreferenceBuilder(
+      preference: settings.selectedFilters,
+      builder: (context, final Set<UnrealObjectFilter> selectedFilters) => DropDownListMenu(
+        minHeight: 100,
+        originTabBuilder: originTabBuilder,
+        children: [
+          for (UnrealObjectFilter filter in UnrealObjectFilterRegistry.filters)
+            ListMenuSimpleItem(
+              key: Key(filter.internalName),
+              title: filter.getDisplayName(context),
+              iconPath: filter.iconPath,
+              bShowCheckbox: true,
+              bIsChecked: selectedFilters.contains(filter),
+              onTap: () => _toggleFilter(filter),
+            ),
+        ],
+      ),
+    );
+  }
+
+  /// Called when the user toggles a filter.
+  void _toggleFilter(UnrealObjectFilter filter) {
+    final Set<UnrealObjectFilter> filters = settings.selectedFilters.getValue();
+    if (filters.contains(filter)) {
+      filters.remove(filter);
+    } else {
+      filters.add(filter);
+    }
+
+    settings.selectedFilters.setValue(filters);
+  }
+}
diff --git a/Engine/Extras/VirtualProduction/EpicStageApp/lib/widgets/screens/main/sidebar/outliner_panel.dart b/Engine/Extras/VirtualProduction/EpicStageApp/lib/widgets/screens/main/sidebar/outliner_panel.dart
new file mode 100644
--- /dev/null
+++ b/Engine/Extras/VirtualProduction/EpicStageApp/lib/widgets/screens/main/sidebar/outliner_panel.dart
@@ -0,0 +1,895 @@
+// Copyright Epic Games, Inc. All Rights Reserved.
+
+import 'dart:async';
+
+import 'package:async/async.dart';
+import 'package:flutter/material.dart';
+import 'package:flutter/services.dart';
+import 'package:flutter_gen/gen_l10n/app_localizations.dart';
+import 'package:logging/logging.dart';
+import 'package:provider/provider.dart';
+import 'package:streaming_shared_preferences/streaming_shared_preferences.dart';
+
+import '../../../../models/actor_data/generic_actor_data.dart';
+import '../../../../models/engine_connection.dart';
+import '../../../../models/settings/main_screen_settings.dart';
+import '../../../../models/settings/outliner_panel_settings.dart';
+import '../../../../models/settings/selected_actor_settings.dart';
+import '../../../../models/unreal_actor_creator.dart';
+import '../../../../models/unreal_actor_manager.dart';
+import '../../../../models/unreal_class_data.dart';
+import '../../../../models/unreal_object_filters.dart';
+import '../../../../models/unreal_transaction_manager.dart';
+import '../../../../models/unreal_types.dart';
+import '../../../../utilities/constants.dart';
+import '../../../../utilities/guarded_refresh_state.dart';
+import '../../../../utilities/preferences_bundle.dart';
+import '../../../../utilities/transient_preference.dart';
+import '../../../../utilities/unreal_colors.dart';
+import '../../../elements/dropdown_button.dart';
+import '../../../elements/dropdown_list_menu.dart';
+import '../../../elements/empty_placeholder.dart';
+import '../../../elements/epic_icon_button.dart';
+import '../../../elements/epic_list_view.dart';
+import '../../../elements/layout/card.dart';
+import '../../../elements/list_menu.dart';
+import '../../../elements/modal.dart';
+import '../../../elements/place_actor_menu.dart';
+import '../../../elements/swipe_revealer.dart';
+import '../../../search_bar.dart';
+import '../stage_app_main_screen.dart';
+import 'outliner_filter_menu.dart';
+
+final _log = Logger('OutlinerPanel');
+
+/// The fraction of the screen's width that should be occupied by the Outliner.
+const _outlinerScreenWidthFraction = 0.28;
+
+/// The maximum width of the Outliner regardless of screen width.
+const _outlinerMaxWidth = 500;
+
+/// The minimum width of the Outliner regardless of screen width.
+const _outlinerMinWidth = 300;
+
+/// Get the width of an outliner panel based on the screen size for the given [context].
+double getOutlinerWidth(BuildContext context) => (MediaQuery.of(context).size.width * _outlinerScreenWidthFraction)
+    .clamp(_outlinerMinWidth, _outlinerMaxWidth)
+    .toDouble();
+
+/// Delete one or more actors, then deselect them once the engine confirms their deletion.
+void _deleteActors(BuildContext context, Iterable<String> actorPaths) async {
+  final transactionManager = Provider.of<UnrealTransactionManager>(context, listen: false);
+
+  if (!transactionManager.beginTransaction('Delete Actors')) {
+    return;
+  }
+
+  final connectionManager = Provider.of<EngineConnectionManager>(context, listen: false);
+
+  // Create a delete request for each actor
+  final requests = actorPaths
+      .map(
+        (String actorPath) => UnrealHttpRequest(
+          url: '/remote/object/call',
+          verb: 'PUT',
+          body: {
+            'generateTransaction': 'false',
+            'objectPath': actorPath,
+            'functionName': 'K2_DestroyActor',
+          },
+        ),
+      )
+      .toList(growable: false);
+
+  // Send a batched message containing all the requests we want to execute
+  final Future<UnrealHttpResponse> batchResponseFuture = connectionManager.sendBatchedHttpRequest(requests);
+
+  transactionManager.endTransaction();
+
+  final UnrealHttpResponse batchResponse = await batchResponseFuture;
+
+  // Check that the batched request succeeded
+  if (batchResponse.code != HttpResponseCode.ok) {
+    _log.warning('Failed to delete actors (error ${batchResponse.code})');
+    return;
+  }
+
+  final List<UnrealHttpResponse?> responses = batchResponse.body;
+  final selectedActorSettings = Provider.of<SelectedActorSettings>(context, listen: false);
+
+  // Check that each delete succeeded and deselect the corresponding actor
+  for (int i = 0; i < responses.length; ++i) {
+    final String actorPath = requests[i].body['objectPath'];
+    final UnrealHttpResponse? response = responses[i];
+    if (response?.code != HttpResponseCode.ok) {
+      _log.warning('Failed to delete actor "$actorPath" (error ${response?.code})');
+    } else {
+      // Actor is confirmed deleted, so deselect it
+      selectedActorSettings.selectActor(actorPath, bShouldSelect: false);
+    }
+  }
+}
+
+/// Send a message to the connected instance of UE requested to rename an actor with [path] to [name], using the
+/// provided [context] to get an instance of [EngineConnectionManager].
+Future<void> _renameActor(BuildContext context, String path, String name) async {
+  final _store = Provider.of<EngineConnectionManager>(context, listen: false);
+
+  // Create a request payload to rename the selected actor @ [path]
+  final request = UnrealHttpRequest(
+    url: '/remote/object/call',
+    verb: 'PUT',
+    body: {
+      "objectPath": path,
+      "functionName": "SetActorLabel",
+      "parameters": {
+        "NewActorLabel": name,
+      },
+      "generateTransaction": true
+    },
+  );
+
+  // Send a message to rename the specified actor at [path].
+  final UnrealHttpResponse res = await _store.sendHttpRequest(request);
+
+  // Check that the request succeeded
+  if (res.code != HttpResponseCode.ok) {
+    _log.warning('Failed to rename actor (error ${res.code})');
+    return;
+  }
+}
+
+/// Get [GenericActorData] for all currently selected actors in the current [context].
+Set<GenericActorData?> _getGenericActorData(BuildContext context) {
+  final SelectedActorSettings _selectedActorSettings = Provider.of<SelectedActorSettings>(context, listen: false);
+  final UnrealActorManager _actorManager = Provider.of<UnrealActorManager>(context, listen: false);
+  return _selectedActorSettings.selectedActors.getValue().map<GenericActorData?>((actor) {
+    UnrealObject? _actor = _actorManager.getActorAtPath(actor);
+    GenericActorData? _actorData = _actor?.getPerClassData<GenericActorData>();
+    return _actorData;
+  }).toSet();
+}
+
+/// Gets the visibility state of the current selected actor/actors in the current [context]
+bool _getSelectedActorsVisibility(BuildContext context) {
+  Set<GenericActorData?> _genericDataForSelectedActors = _getGenericActorData(context);
+
+  // If there is exactly one actor on the list.
+  if (_genericDataForSelectedActors.length == 1) {
+    return _genericDataForSelectedActors.first?.bIsHiddenInGame == false;
+    // If there are multiple actors on the list.
+  } else if (_genericDataForSelectedActors.isNotEmpty && _genericDataForSelectedActors.length > 1) {
+    /// Whether all selected actor have same visibility state.
+    bool bIsAllActorVisibilityStateSame;
+    bIsAllActorVisibilityStateSame = _genericDataForSelectedActors
+        .every((element) => element?.bIsHiddenInGame == _genericDataForSelectedActors.toList()[0]?.bIsHiddenInGame);
+    return bIsAllActorVisibilityStateSame ? !_genericDataForSelectedActors.first!.bIsHiddenInGame : true;
+    // If all above conditions are not met, we just want to return false.
+  } else {
+    return false;
+  }
+}
+
+/// A sidebar panel that lets the user view and interact with the list of actors.
+class OutlinerPanel extends StatefulWidget {
+  const OutlinerPanel({
+    Key? key,
+    this.focusActor,
+    this.canFocusActor,
+  }) : super(key: key);
+
+  /// A function to call when the user wants to focus the stage map on a particular actor.
+  /// If not provided, the user won't be given the option to do so.
+  final void Function(UnrealObject)? focusActor;
+
+  /// A function that checks whether the given actor can be focused.
+  /// If [focusActor] is provided but this isn't, it's assumed that all actors can be focused.
+  final bool Function(UnrealObject)? canFocusActor;
+
+  @override
+  State<OutlinerPanel> createState() => _OutlinerPanelState();
+}
+
+class _OutlinerPanelState extends State<OutlinerPanel> {
+  late final UnrealActorManager _actorManager;
+  late final OutlinerPanelSettings _outlinerSettings;
+
+  final TextEditingController _searchTextController = TextEditingController();
+
+  /// Set of actors that we're listening to for any updates (e.g. to class information).
+  Set<UnrealObject> _listenedActors = {};
+
+  /// Filtered list of actor names to show.
+  List<UnrealObject> _filteredActors = [];
+
+  /// User-supplied text used to filter the list of actors by name.
+  String _filterText = '';
+
+  /// Subscription to settings that we want to watch.
+  StreamSubscription? _settingsSubscription;
+
+  /// Whether selected actor/actors is visible or not.
+  bool _bIsSelectedActorsVisible = true;
+
+  @override
+  void initState() {
+    super.initState();
+
+    _actorManager = Provider.of<UnrealActorManager>(context, listen: false);
+    _outlinerSettings = OutlinerPanelSettings(PreferencesBundle.of(context));
+
+    // Update filtered actors when sort settings change
+    _settingsSubscription = StreamGroup.merge([
+      _outlinerSettings.actorSortMode,
+      _outlinerSettings.bShouldSortActorsDescending,
+      _outlinerSettings.selectedFilters,
+    ]).listen((_) => _updateFilteredActors());
+
+    for (final String className in controllableClassNames) {
+      _actorManager.watchClassName(className, _onActorsChanged);
+    }
+
+    _actorManager.currentRootActorLightCards.addListener(_updateFilteredActors);
+  }
+
+  @override
+  void dispose() {
+    _actorManager.currentRootActorLightCards.removeListener(_updateFilteredActors);
+
+    for (final String className in controllableClassNames) {
+      _actorManager.stopWatchingClassName(className, _onActorsChanged);
+    }
+
+    for (final UnrealObject actor in _listenedActors) {
+      actor.removeListener(_updateFilteredActors);
+    }
+    _listenedActors.clear();
+
+    _settingsSubscription?.cancel();
+
+    super.dispose();
+  }
+
+  @override
+  Widget build(BuildContext context) {
+    final num width = getOutlinerWidth(context);
+    final selectedActorSettings = Provider.of<SelectedActorSettings>(context, listen: false);
+
+    return Provider<OutlinerPanelSettings>(
+      create: (context) => _outlinerSettings,
+      child: MediaQuery.removePadding(
+        context: context,
+        removeBottom: true,
+        child: SizedBox(
+          width: width.toDouble(),
+          child: Card(
+            child: Column(
+              children: [
+                CardSmallHeader(title: AppLocalizations.of(context)!.outlinerTitle),
+
+                Container(
+                  height: sectionMargin,
+                  color: Theme.of(context).colorScheme.surfaceVariant,
+                ),
+
+                // Search bar
+                CardSubHeader(
+                  height: 52,
+                  padding: const EdgeInsets.symmetric(horizontal: 16, vertical: 8),
+                  child: EpicSearchBar(
+                    controller: _searchTextController,
+                    onChanged: _onFilterTextChanged,
+                  ),
+                ),
+
+                // Button bar
+                CardSubHeader(
+                  child: Row(
+                    mainAxisAlignment: MainAxisAlignment.spaceBetween,
+                    children: [
+                      const OutlinerFilterButton(),
+                      LightCardVisibilityToggle(bIsVisible: _bIsSelectedActorsVisible),
+                      PreferenceBuilder(
+                        preference: selectedActorSettings.bIsInMultiSelectMode,
+                        builder: (BuildContext context, bool bIsInMultiSelectMode) {
+                          return CardSubHeaderButton(
+                            iconPath: 'assets/images/icons/multi_select.svg',
+                            tooltipMessage: AppLocalizations.of(context)!.toggleMultiselectTooltip,
+                            bIsToggledOn: bIsInMultiSelectMode,
+                            onPressed: () => selectedActorSettings.bIsInMultiSelectMode.setValue(!bIsInMultiSelectMode),
+                          );
+                        },
+                      ),
+                      ModalDropdownButton(
+                        buttonBuilder: (context, state) => Material(
+                          color: Colors.transparent,
+                          child: CardSubHeaderButton(
+                            iconPath: 'assets/images/icons/ellipsis.svg',
+                            tooltipMessage: AppLocalizations.of(context)!.moreActions,
+                            bIsToggledOn: state != ModalDropdownButtonState.closed,
+                            bIsVisualOnly: true,
+                          ),
+                        ),
+                        menuBuilder: (context, originTabBuilder) => TransientPreferenceBuilder(
+                          preference: selectedActorSettings.selectedActors,
+                          builder: (BuildContext context, Set<String> selectedActors) {
+                            final bool bHasSingleTarget = selectedActors.length == 1;
+                            final bool bShowFocus = widget.focusActor != null;
+                            late final bool bCanFocus;
+                            if (bShowFocus && bHasSingleTarget) {
+                              if (widget.canFocusActor == null) {
+                                // Assume we can focus any single target
+                                bCanFocus = true;
+                              } else {
+                                // Check that we can focus the current target
+                                final UnrealObject? singleTarget = _actorManager.getActorAtPath(selectedActors.first);
+                                bCanFocus = singleTarget != null && widget.canFocusActor!.call(singleTarget);
+                              }
+                            } else {
+                              // No single target/no focus function, so we can never focus
+                              bCanFocus = false;
+                            }
+                            return DropDownListMenu(
+                              originTabBuilder: originTabBuilder,
+                              children: [
+                                if (bShowFocus)
+                                  ListMenuSimpleItem(
+                                    iconPath: 'assets/images/icons/focus.svg',
+                                    title: AppLocalizations.of(context)!.outlinerFocusSelected,
+                                    onTap: () {
+                                      _focusSelectedActor();
+                                      Navigator.of(context).pop();
+                                    },
+                                    bIsEnabled: bCanFocus,
+                                  ),
+                                ListMenuSimpleItem(
+                                  iconPath: 'assets/images/icons/paste.svg',
+                                  title: AppLocalizations.of(context)!.outlinerDuplicateSelected,
+                                  onTap: () {
+                                    _duplicateSelectedActors();
+                                    Navigator.of(context).pop();
+                                  },
+                                ),
+                                ListMenuSimpleItem(
+                                  iconPath: 'assets/images/icons/trash.svg',
+                                  title: AppLocalizations.of(context)!.outlinerDeleteSelected,
+                                  onTap: () {
+                                    _deleteSelectedActors();
+                                    Navigator.of(context).pop();
+                                  },
+                                ),
+                                ListMenuSimpleItem(
+                                  iconPath: 'assets/images/icons/edit.svg',
+                                  title: AppLocalizations.of(context)!.outlinerRenameSelected,
+                                  onTap: () {
+                                    Navigator.of(context).pop();
+                                    _showRenameSelectedActorDialog();
+                                  },
+                                  bIsEnabled: selectedActors.length == 1,
+                                ),
+                              ],
+                            );
+                          },
+                        ),
+                      ),
+                    ],
+                  ),
+                ),
+
+                //List of actors or onboarding placeholder
+                Expanded(
+                  child: _filteredActors.isNotEmpty
+                      ? EpicListView(
+                          padding: EdgeInsets.only(top: 4),
+                          itemCount: _filteredActors.length,
+                          itemBuilder: (BuildContext context, int actorIndex) {
+                            final UnrealObject actor = _filteredActors[actorIndex];
+
+                            return _OutlinerPanelActor(
+                              key: Key(actor.path),
+                              actor: actor,
+                              focusActor: widget.focusActor,
+                              canFocusActor: widget.canFocusActor,
+                              visibilityStateCallback: (state) => setState(() => _bIsSelectedActorsVisible = state),
+                            );
+                          },
+                        )
+                      : EmptyPlaceholder(
+                          message: AppLocalizations.of(context)!.outlinerEmptyMessage,
+                          button: EpicWideButton(
+                            text: AppLocalizations.of(context)!.addActorButton,
+                            iconPath: 'assets/images/icons/plus.svg',
+                            iconColor: UnrealColors.highlightGreen,
+                            onPressed: () => DropDownListMenu.showAtWidget(
+                              widgetKey:
+                                  Provider.of<StageAppMainScreenKeys>(context, listen: false).placeActorsButtonKey,
+                              builder: (context) => PlaceActorDropDownMenu(
+                                bIsFromLongPress: false,
+                              ),
+                            ),
+                          ),
+                        ),
+                ),
+              ],
+            ),
+          ),
+        ),
+      ),
+    );
+  }
+
+  /// Called when the user changes the filter text input.
+  void _onFilterTextChanged(String newValue) {
+    _filterText = newValue;
+    _updateFilteredActors();
+  }
+
+  /// Update the filtered and sorted list of actors in the dropdown menu.
+  void _updateFilteredActors() {
+    if (!mounted) {
+      return;
+    }
+
+    // Get a set of all valid actors to display
+    final Set<UnrealObject> validActors = {};
+    for (final String className in controllableClassNames) {
+      for (final UnrealObject actor in _actorManager.getActorsOfClass(className)) {
+        validActors.add(actor);
+      }
+    }
+
+    // Stop listening to actors no longer in the list
+    final List<UnrealObject> removedActors = [];
+    for (final UnrealObject actor in _listenedActors) {
+      if (!validActors.contains(actor)) {
+        actor.removeListener(_updateFilteredActors);
+        removedActors.add(actor);
+      }
+    }
+
+    for (final UnrealObject actor in removedActors) {
+      _listenedActors.remove(actor);
+    }
+
+    // Start listening to new actors in the list
+    for (final UnrealObject actor in validActors) {
+      if (!_listenedActors.contains(actor)) {
+        actor.addListener(_updateFilteredActors);
+        _listenedActors.add(actor);
+      }
+    }
+
+    final OutlinerActorSortMode sortMode = _outlinerSettings.actorSortMode.getValue();
+    final Set<UnrealObjectFilter> filters = _outlinerSettings.selectedFilters.getValue();
+
+    setState(() {
+      // Filter out actors we shouldn't display
+      _filteredActors = validActors.where((actor) => _filterActor(actor, filters)).toList(growable: false);
+
+      int Function(UnrealObject, UnrealObject) compareFunction;
+      switch (sortMode) {
+        case OutlinerActorSortMode.name:
+          compareFunction = _compareActorsByName;
+          break;
+
+        case OutlinerActorSortMode.recent:
+          compareFunction = _compareActorsByRecency;
+          break;
+
+        default:
+          throw Exception('No comparison function for sort type ${sortMode}');
+      }
+
+      // Reverse the comparison function if we're sorting in ascending order.
+      if (!_outlinerSettings.bShouldSortActorsDescending.getValue()) {
+        final int Function(UnrealObject, UnrealObject) baseCompareFunction = compareFunction;
+        compareFunction = (UnrealObject a, UnrealObject b) => -baseCompareFunction(a, b);
+      }
+
+      _filteredActors.sort(compareFunction);
+    });
+  }
+
+  /// Return true if the actor with the given path should be included in the list shown to the user.
+  bool _filterActor(UnrealObject actorData, Set<UnrealObjectFilter> filters) {
+    // Check that the actor name contains the filter text
+    if (_filterText.isNotEmpty && !actorData.name.toLowerCase().contains(_filterText.toLowerCase())) {
+      return false;
+    }
+
+    // If the actor is a light card, only show it if it belongs to the current DCRA.
+    if (actorData.unrealClass?.isA(UnrealClassRegistry.lightCard) == true &&
+        !_actorManager.currentRootActorLightCards.value.contains(actorData)) {
+      return false;
+    }
+
+    // Check that the actor passes at least one filter, if any are present
+    if (filters.isNotEmpty) {
+      bool bPassesFilter = false;
+      for (final UnrealObjectFilter filter in filters) {
+        if (filter.passes(actorData)) {
+          bPassesFilter = true;
+          break;
+        }
+      }
+
+      if (!bPassesFilter) {
+        return false;
+      }
+    }
+
+    return true;
+  }
+
+  /// Comparison function to sort actors by name.
+  int _compareActorsByName(UnrealObject actorA, UnrealObject actorB) {
+    return actorA.name.compareTo(actorB.name);
+  }
+
+  /// Comparison function to sort actors by recency.
+  int _compareActorsByRecency(UnrealObject actorA, UnrealObject actorB) {
+    final selectedActorSettings = Provider.of<SelectedActorSettings>(context, listen: false);
+
+    // First, check if the actors are currently selected
+    final bool bIsASelected = selectedActorSettings.isActorSelected(actorA.path);
+    final bool bIsBSelected = selectedActorSettings.isActorSelected(actorB.path);
+
+    if (bIsASelected && bIsBSelected) {
+      // Use name as tie-breaker if both are selected
+      return _compareActorsByName(actorA, actorB);
+    }
+
+    if (bIsASelected) {
+      return -1;
+    }
+
+    if (bIsBSelected) {
+      return 1;
+    }
+
+    // Neither actor is currently selected, so compare based on when they were last selected (if ever)
+    final DateTime? selectedTimeA = selectedActorSettings.getActorLastSelectedTime(actorB.path);
+    final DateTime? selectedTimeB = selectedActorSettings.getActorLastSelectedTime(actorB.path);
+
+    if (selectedTimeA == null && selectedTimeB == null) {
+      // Use name as tie-breaker if neither have ever been selected
+      return _compareActorsByName(actorA, actorB);
+    }
+
+    if (selectedTimeA == null) {
+      return 1;
+    }
+
+    if (selectedTimeB == null) {
+      return -1;
+    }
+
+    return -selectedTimeA.compareTo(selectedTimeB);
+  }
+
+  /// Send a message to the engine requesting the deletion of the selected actors.
+  void _deleteSelectedActors() {
+    final selectedActorSettings = Provider.of<SelectedActorSettings>(context, listen: false);
+    final Set<String> selectedActors = selectedActorSettings.selectedActors.getValue();
+    if (selectedActors.isEmpty) {
+      return;
+    }
+
+    _deleteActors(context, selectedActors);
+  }
+
+  /// Duplicate the selected actors.
+  void _duplicateSelectedActors() {
+    final actorCreator = Provider.of<UnrealActorCreator>(context, listen: false);
+    actorCreator.duplicateSelectedActors();
+  }
+
+  /// Focus the selected actor (if possible, and if only one is selected).
+  void _focusSelectedActor() {
+    if (widget.focusActor == null) {
+      return;
+    }
+
+    final UnrealObject? actor = _getSingleSelectedActor();
+
+    if (actor != null) {
+      widget.focusActor!(actor);
+    }
+  }
+
+  /// Show a dialog to rename the selected actor (if possible, and if only one is selected).
+  void _showRenameSelectedActorDialog() {
+    final UnrealObject? actor = _getSingleSelectedActor();
+
+    if (actor == null) {
+      return;
+    }
+
+    GenericModalDialogRoute.showDialog(
+      context: context,
+      builder: (context) => StringTextInputModalDialog(
+        title: AppLocalizations.of(context)!.renameActorModalTitle,
+        initialValue: actor.name,
+        handleResult: (TextInputModalDialogResult<String> result) async {
+          if (result.action == TextInputModalDialogAction.apply) {
+            await _renameActor(context, actor.path, result.value!);
+          }
+        },
+      ),
+    );
+  }
+
+  /// Get the single selected actor, or null if there isn't one/multiple are selected.
+  UnrealObject? _getSingleSelectedActor() {
+    final selectedActorSettings = Provider.of<SelectedActorSettings>(context, listen: false);
+    final Set<String> selectedActors = selectedActorSettings.selectedActors.getValue();
+
+    if (selectedActors.length != 1) {
+      return null;
+    }
+
+    return _actorManager.getActorAtPath(selectedActors.first);
+  }
+
+  /// Called when the list of actors provided by the engine changes.
+  void _onActorsChanged(ActorUpdateDetails details) {
+    _updateFilteredActors();
+  }
+}
+
+/// An entry for an actor in the outliner panel's actor list.
+class _OutlinerPanelActor extends StatefulWidget {
+  const _OutlinerPanelActor(
+      {Key? key, required this.actor, this.focusActor, this.canFocusActor, required this.visibilityStateCallback})
+      : super(key: key);
+
+  /// The actor this represents.
+  final UnrealObject actor;
+
+  /// The function to call when the user wants to focus this actor.
+  final void Function(UnrealObject)? focusActor;
+
+  /// A function that checks whether the given actor can be focused.
+  /// If [focusActor] is provided but this isn't, it's assumed that all actors can be focused.
+  final bool Function(UnrealObject)? canFocusActor;
+
+  /// value changed callback to get initial visibility state of an actor when click/selected.
+  final ValueChanged<bool> visibilityStateCallback;
+
+  @override
+  State<StatefulWidget> createState() => _OutlinerPanelActorState();
+}
+
+class _OutlinerPanelActorState extends State<_OutlinerPanelActor> with GuardedRefreshState, TickerProviderStateMixin {
+  /// Data for the actor this represents.
+  late final GenericActorData? _actorData;
+
+  /// Controller for the animation shown when this is long pressed.
+  late final AnimationController _longPressAnimationController;
+
+  /// The opacity animation shown when this is long pressed.
+  late final Animation<double> _longPressOpacityAnimation;
+
+  @override
+  void initState() {
+    super.initState();
+
+    _actorData = widget.actor.getPerClassData<GenericActorData>();
+    _actorData?.addListener(guardedRefresh);
+
+    _longPressAnimationController = AnimationController(
+      vsync: this,
+      duration: const Duration(milliseconds: 500),
+    );
+
+    _longPressOpacityAnimation = TweenSequence([
+      TweenSequenceItem<double>(
+        tween: Tween(begin: 0, end: 1),
+        weight: 0.2,
+      ),
+      TweenSequenceItem<double>(
+        tween: ConstantTween(1),
+        weight: 4,
+      ),
+      TweenSequenceItem<double>(
+        tween: Tween(begin: 1, end: 0),
+        weight: 10,
+      ),
+    ]).animate(
+      CurvedAnimation(
+        parent: _longPressAnimationController,
+        curve: Curves.easeInOut,
+      ),
+    );
+  }
+
+  @override
+  void dispose() {
+    _actorData?.removeListener(guardedRefresh);
+    _longPressAnimationController.dispose();
+    super.dispose();
+  }
+
+  @override
+  Widget build(BuildContext context) {
+    final selectedActorSettings = Provider.of<SelectedActorSettings>(context, listen: false);
+
+    final body = Stack(
+      children: [
+        SwipeRevealer(
+          backgroundPadding: const EdgeInsets.only(top: 2),
+          onDeleted: () => _deleteActors(context, {widget.actor.path}),
+          leftSwipeActionBuilder: (context, onFinished) => CardListTileSwipeAction(
+            iconPath: 'assets/images/icons/trash.svg',
+            color: UnrealColors.highlightRed,
+            onPressed: () => onFinished(bDeleteItem: true),
+          ),
+          rightSwipeActionBuilder: (context, onFinished) => CardListTileSwipeAction(
+            iconPath: (_actorData?.bIsHiddenInGame == true)
+                ? 'assets/images/icons/hidden_in_game.svg'
+                : 'assets/images/icons/visible_in_game.svg',
+            color: UnrealColors.gray22,
+            onPressed: () {
+              if (_actorData != null) {
+                _actorData!.bIsHiddenInGame = !_actorData!.bIsHiddenInGame;
+              }
+              onFinished();
+            },
+          ),
+          child: TransientPreferenceBuilder(
+            preference: selectedActorSettings.selectedActors,
+            builder: (BuildContext context, Set<String> selectedActors) {
+              final bool bIsSelected = selectedActors.contains(widget.actor.path);
+              return CardListTile(
+                title: widget.actor.name,
+                bIsSelected: bIsSelected,
+                iconPath: widget.actor.getIconPath(),
+                bDeEmphasize: _actorData?.bIsHiddenInGame ?? false,
+                onTap: () {
+                  selectedActorSettings.selectActor(widget.actor.path, bShouldSelect: !bIsSelected);
+                  widget.visibilityStateCallback(_getSelectedActorsVisibility(context));
+                },
+              );
+            },
+          ),
+        ),
+        Positioned(
+          left: 0,
+          right: 0,
+          top: 2,
+          bottom: 0,
+          child: IgnorePointer(
+            child: AnimatedBuilder(
+              animation: _longPressOpacityAnimation,
+              builder: (context, child) => Container(
+                decoration: BoxDecoration(
+                  color: Colors.white.withOpacity(_longPressOpacityAnimation.value * 0.3),
+                ),
+              ),
+            ),
+          ),
+        ),
+      ],
+    );
+
+    if (widget.focusActor == null) {
+      return body;
+    }
+
+    return GestureDetector(
+      onLongPress: _onLongPress,
+      child: body,
+      behavior: HitTestBehavior.opaque,
+    );
+  }
+
+  /// Called when the user long-presses on this.
+  void _onLongPress() {
+    if (!(widget.canFocusActor?.call(widget.actor) ?? true)) {
+      return;
+    }
+
+    _longPressAnimationController.forward(from: 0);
+
+    widget.focusActor!(widget.actor);
+    HapticFeedback.vibrate();
+  }
+}
+
+/// Button that enables/disables the outliner panel.
+class OutlinerToggleButton extends StatelessWidget {
+  const OutlinerToggleButton({
+    Key? key,
+    this.bEnabled = true,
+  }) : super(key: key);
+
+  /// If false, the button can't be toggled by the user.
+  final bool bEnabled;
+
+  @override
+  Widget build(BuildContext context) {
+    final mainScreenSettings = Provider.of<MainScreenSettings>(context);
+
+    return PreferenceBuilder(
+      preference: mainScreenSettings.bIsOutlinerPanelOpen,
+      builder: (BuildContext context, bool bIsOutlinerPanelOpen) {
+        return EpicIconButton(
+          iconPath: 'assets/images/icons/outliner.svg',
+          tooltipMessage: AppLocalizations.of(context)!.outlinerTitle,
+          bIsToggledOn: bIsOutlinerPanelOpen,
+          onPressed: bEnabled
+              ? () {
+                  mainScreenSettings.bIsOutlinerPanelOpen.setValue(!bIsOutlinerPanelOpen);
+                }
+              : null,
+        );
+      },
+    );
+  }
+}
+
+/// Button to make actors visible or hidden.
+class LightCardVisibilityToggle extends StatefulWidget {
+  const LightCardVisibilityToggle({Key? key, required this.bIsVisible}) : super(key: key);
+  final bool bIsVisible;
+
+  @override
+  State<LightCardVisibilityToggle> createState() => _LightCardVisibilityToggleState();
+}
+
+class _LightCardVisibilityToggleState extends State<LightCardVisibilityToggle> with GuardedRefreshState {
+  /// Whether the button shows the visibility icon or the hidden icon.
+  late bool bIsVisible;
+
+  @override
+  void initState() {
+    bIsVisible = widget.bIsVisible;
+    super.initState();
+  }
+
+  @override
+  Widget build(BuildContext context) {
+    bIsVisible = _getSelectedActorsVisibility(context);
+    return CardSubHeaderButton(
+      iconPath: bIsVisible ? 'assets/images/icons/visible_in_game.svg' : 'assets/images/icons/hidden_in_game.svg',
+      tooltipMessage: AppLocalizations.of(context)!.outlinerToggleVisibility,
+      bIsToggledOn: false,
+      onPressed: () => _toggleVisibility(),
+    );
+  }
+
+  /// Callback function to toggle visibility of currently selected actor/actors.
+  void _toggleVisibility() {
+    Set<GenericActorData?> _data = _getGenericActorData(context);
+
+    /// Assert if the actors are all visible/hidden or they are mixed of both hidden and visible actors.
+    bool bAssert = _data.every((element) => element?.bIsHiddenInGame == _data.toList()[0]?.bIsHiddenInGame);
+
+    //If there is just one selected actor.
+    if (_data.length == 1) {
+      GenericActorData? _actorData = _data.first;
+      _setState(_actorData);
+    } else if (_data.length > 1 && bAssert) {
+      for (GenericActorData? data in _data) _setState(data);
+    } else {
+      for (GenericActorData? data in _data) {
+        if (data != null) {
+          data.addListener(guardedRefresh);
+          data.bIsHiddenInGame = true;
+          data.removeListener(guardedRefresh);
+          bIsVisible = _getSelectedActorsVisibility(context);
+        }
+      }
+    }
+  }
+
+  /// Set's the visibility state of an actor with [data] to be either visible or hidden.
+  void _setState(GenericActorData? data) {
+    if (data != null) {
+      data.addListener(guardedRefresh);
+      data.bIsHiddenInGame = !data.bIsHiddenInGame;
+      data.removeListener(guardedRefresh);
+      bIsVisible = _getSelectedActorsVisibility(context);
+    }
+  }
+}
diff --git a/Engine/Extras/VirtualProduction/EpicStageApp/lib/widgets/screens/main/stage_app_main_screen.dart b/Engine/Extras/VirtualProduction/EpicStageApp/lib/widgets/screens/main/stage_app_main_screen.dart
new file mode 100644
--- /dev/null
+++ b/Engine/Extras/VirtualProduction/EpicStageApp/lib/widgets/screens/main/stage_app_main_screen.dart
@@ -0,0 +1,228 @@
+// Copyright Epic Games, Inc. All Rights Reserved.
+
+import 'package:flutter/material.dart';
+import 'package:flutter_gen/gen_l10n/app_localizations.dart';
+import 'package:provider/provider.dart';
+
+import '../../../models/settings/main_screen_settings.dart';
+import '../../../models/settings/selected_actor_settings.dart';
+import '../../../models/unreal_actor_manager.dart';
+import '../../../models/unreal_types.dart';
+import '../../../utilities/constants.dart';
+import '../../../utilities/debug_utilities.dart';
+import '../../../utilities/unreal_colors.dart';
+import '../../elements/floating_map_preview.dart';
+import '../../elements/layout/lazy_tab_view.dart';
+import '../connect/views/n_display_selector_dialog.dart';
+import 'tabs/main_screen_tabs.dart';
+import 'toolbar/main_screen_toolbar.dart';
+
+/// The main screen of the app shown once the user is connected to the engine.
+class StageAppMainScreen extends StatefulWidget {
+  const StageAppMainScreen({Key? key}) : super(key: key);
+
+  static const String route = '/main';
+
+  @override
+  State<StageAppMainScreen> createState() => _StageAppMainScreenState();
+}
+
+class _StageAppMainScreenState extends State<StageAppMainScreen>
+    with SingleTickerProviderStateMixin, WidgetsBindingObserver {
+  late final UnrealActorManager _actorManager;
+  late final TabController _tabController;
+
+  /// Global key that refers to the toolbar's place actors button.
+  final GlobalKey _placeActorsButtonKey = GlobalKey();
+
+  /// Whether the outliner can be toggled in the current tab.
+  bool _bEnableOutlinerToggle = true;
+
+  /// Whether to show the floating stage map preview.
+  bool _bShowMapPreview = false;
+
+  @override
+  void initState() {
+    super.initState();
+
+    _tabController = TabController(
+      animationDuration: Duration.zero,
+      length: MainScreenTabs.tabConfigs.length,
+      vsync: this,
+    );
+    _tabController.addListener(_onActiveTabChanged);
+
+    final int startTab = Provider.of<MainScreenSettings>(context, listen: false).selectedTab.getValue();
+
+    if (startTab >= 0 && startTab < _tabController.length) {
+      _tabController.index = startTab;
+    } else {
+      _tabController.index = 0;
+    }
+
+    _actorManager = Provider.of<UnrealActorManager>(context, listen: false);
+    _actorManager.watchClassName(nDisplayRootActorClassName, _onRootActorUpdate);
+
+    final List<Future<Set<UnrealObject>>> initialActorFutures = [];
+
+    /// Subscribe to classes that we always want to know about regardless of the current tab.
+    /// This prevents us from constantly sending unsubscribe and resubscribe messages when changing tabs.
+    for (final String className in controllableClassNames) {
+      _actorManager.watchClassName(className, _onControllableActorUpdate);
+      initialActorFutures.add(_actorManager.getInitialActorsOfClass(className));
+    }
+
+    Future.wait(initialActorFutures).then(_handleInitialControllableActorList);
+  }
+
+  @override
+  void dispose() {
+    _tabController.dispose();
+    _actorManager.stopWatchingClassName(nDisplayRootActorClassName, _onRootActorUpdate);
+
+    for (final String className in controllableClassNames) {
+      _actorManager.stopWatchingClassName(className, _onControllableActorUpdate);
+    }
+
+    super.dispose();
+  }
+
+  @override
+  Widget build(BuildContext context) {
+    return Scaffold(
+      resizeToAvoidBottomInset: false,
+      appBar: MainScreenToolbar(
+        tabController: _tabController,
+        bEnableOutlinerToggle: _bEnableOutlinerToggle,
+        placeActorButtonKey: _placeActorsButtonKey,
+      ),
+      body: Provider(
+        create: (_) => StageAppMainScreenKeys(placeActorsButtonKey: _placeActorsButtonKey),
+        child: Stack(
+          clipBehavior: Clip.none,
+          children: [
+            Container(
+              color: UnrealColors.gray14,
+              padding: EdgeInsets.only(left: cardMargin, right: cardMargin, bottom: cardMargin),
+              key: Key('Tab View'),
+              child: Container(
+                decoration: BoxDecoration(
+                  color: Theme.of(context).colorScheme.background,
+                  borderRadius: BorderRadius.circular(outerCornerRadius),
+                ),
+                clipBehavior: Clip.antiAlias,
+                child: LazyTabView(
+                  controller: _tabController,
+                  builder: MainScreenTabs.createTabContents,
+                  keepAlive: MainScreenTabs.shouldKeepTabAlive,
+                ),
+              ),
+            ),
+            if (_bShowMapPreview) const FloatingMapPreview(key: Key('Map Preview')),
+          ],
+        ),
+      ),
+    );
+  }
+
+  /// Called when the actors in the engine have changed.
+  void _onRootActorUpdate(ActorUpdateDetails details) async {
+    if (!mounted) {
+      return;
+    }
+
+    final selectedActorSettings = Provider.of<SelectedActorSettings>(context, listen: false);
+    final Set<UnrealObject> rootActors = _actorManager.getActorsOfClass(nDisplayRootActorClassName);
+
+    if (rootActors.any((actor) => actor.path == selectedActorSettings.displayClusterRootPath.getValue())) {
+      // Our selected actor is still valid
+      return;
+    }
+
+    if (rootActors.isEmpty) {
+      // Wait until end of frame so the root actor select menu, if any, has time to pop
+      WidgetsBinding.instance.addPostFrameCallback((_) {
+        if (mounted) {
+          showDebugAlert(AppLocalizations.of(context)!.warningNoRootActors);
+        }
+      });
+      return;
+    }
+
+    if (rootActors.length > 1) {
+      // Prompt the user to select one of the new root actors
+      final UnrealObject? newActor = await NDisplaySelectorDialog.showIfNotOpen(context, bIsAlreadyConnected: true);
+      Provider.of<SelectedActorSettings>(context, listen: false).displayClusterRootPath.setValue(newActor?.path ?? '');
+      return;
+    }
+
+    // There's only one new root actor, so use that one
+    Provider.of<SelectedActorSettings>(context, listen: false).displayClusterRootPath.setValue(rootActors.first.path);
+  }
+
+  /// Called when the active tab changes.
+  void _onActiveTabChanged() {
+    Provider.of<MainScreenSettings>(context, listen: false).selectedTab.setValue(_tabController.index);
+
+    final bool bNewEnableOutlinerToggle = MainScreenTabs.shouldEnableOutlinerToggle(_tabController.index);
+    if (bNewEnableOutlinerToggle != _bEnableOutlinerToggle) {
+      setState(() {
+        _bEnableOutlinerToggle = bNewEnableOutlinerToggle;
+      });
+    }
+
+    final bool bNewShowMapPreview = MainScreenTabs.shouldShowMapPreview(_tabController.index);
+    if (bNewShowMapPreview != _bShowMapPreview) {
+      setState(() {
+        _bShowMapPreview = bNewShowMapPreview;
+      });
+    }
+  }
+
+  /// Called when a controllable actor has changed.
+  void _onControllableActorUpdate(ActorUpdateDetails details) {
+    if (!mounted || details.bIsDueToDisconnect) {
+      // Don't deselect actors that are reported as deleted due to disconnect since we can re-select then on reconnect
+      return;
+    }
+
+    // Deselect any deleted actors
+    final selectedActorSettings = Provider.of<SelectedActorSettings>(context, listen: false);
+
+    for (final UnrealObject actor in details.deletedActors) {
+      if (selectedActorSettings.isActorSelected(actor.path)) {
+        selectedActorSettings.selectActor(actor.path, bShouldSelect: false);
+      }
+    }
+  }
+
+  /// Handle the initial list of [actors] received from the engine.
+  void _handleInitialControllableActorList(List<Set<UnrealObject>> actors) {
+    if (!mounted) {
+      return;
+    }
+
+    final selectedActorSettings = Provider.of<SelectedActorSettings>(context, listen: false);
+
+    // Deselect any actors that no longer exist
+    final staleActorPaths = Set<String>();
+
+    for (final String actorPath in selectedActorSettings.selectedActors.getValue()) {
+      if (_actorManager.getActorAtPath(actorPath) == null) {
+        staleActorPaths.add(actorPath);
+      }
+    }
+
+    for (final String actorPath in staleActorPaths) {
+      selectedActorSettings.selectActor(actorPath, bShouldSelect: false);
+    }
+  }
+}
+
+/// Container for global keys that refer to elements of the main screen.
+class StageAppMainScreenKeys {
+  const StageAppMainScreenKeys({required this.placeActorsButtonKey});
+
+  /// Key for the place actors button in the main screen's toolbar.
+  final GlobalKey placeActorsButtonKey;
+}
diff --git a/Engine/Extras/VirtualProduction/EpicStageApp/lib/widgets/screens/main/tabs/base_color_tab.dart b/Engine/Extras/VirtualProduction/EpicStageApp/lib/widgets/screens/main/tabs/base_color_tab.dart
new file mode 100644
--- /dev/null
+++ b/Engine/Extras/VirtualProduction/EpicStageApp/lib/widgets/screens/main/tabs/base_color_tab.dart
@@ -0,0 +1,408 @@
+// Copyright Epic Games, Inc. All Rights Reserved.
+
+import 'package:flutter/material.dart';
+import 'package:flutter_gen/gen_l10n/app_localizations.dart';
+import 'package:provider/provider.dart';
+import 'package:streaming_shared_preferences/streaming_shared_preferences.dart';
+
+import '../../../../models/settings/base_color_tab_settings.dart';
+import '../../../../models/settings/color_wheel_settings.dart';
+import '../../../../models/unreal_types.dart';
+import '../../../../utilities/constants.dart';
+import '../../../../utilities/preferences_bundle.dart';
+import '../../../../utilities/unreal_utilities.dart';
+import '../../../elements/color_grading_wheel.dart';
+import '../../../elements/delta_slider.dart';
+import '../../../elements/dropdown_button.dart';
+import '../../../elements/dropdown_text.dart';
+import '../../../elements/selector_bar.dart';
+
+/// Modes the base color tab can operate in.
+enum BaseColorTabMode {
+  /// Edit FColor properties in Unreal.
+  color,
+
+  /// Edit FVector color grading property groups in Unreal.
+  colorGrading,
+
+  /// Edit post-process volumes in Unreal.
+  postProcess,
+}
+
+/// Property names for each of the highlight color ranges.
+const Map<RangeLimitType, String> _highlightRangeLimitPropertyNames = {
+  RangeLimitType.min: 'Min',
+  RangeLimitType.max: 'Max',
+};
+
+/// Property names for each color grading range.
+const Map<ColorGradingRange, String> _colorGradingRangePropertyNames = {
+  ColorGradingRange.global: 'Global',
+  ColorGradingRange.shadows: 'Shadows',
+  ColorGradingRange.midtones: 'Midtones',
+  ColorGradingRange.highlights: 'Highlights',
+};
+
+/// Settings to use for each color grading subproperty.
+final Map<ColorGradingSubproperty, ColorGradingSubpropertyConfig> _colorGradingSubpropertyConfig = {
+  ColorGradingSubproperty.gain: ColorGradingSubpropertyConfig(
+    name: 'Gain',
+    getDisplayName: (context) => AppLocalizations.of(context)!.colorGradingSubpropertyGain,
+  ),
+  ColorGradingSubproperty.gamma: ColorGradingSubpropertyConfig(
+    name: 'Gamma',
+    getDisplayName: (context) => AppLocalizations.of(context)!.colorGradingSubpropertyGamma,
+  ),
+  ColorGradingSubproperty.saturation: ColorGradingSubpropertyConfig(
+    name: 'Saturation',
+    getDisplayName: (context) => AppLocalizations.of(context)!.colorGradingSubpropertySaturation,
+  ),
+  ColorGradingSubproperty.contrast: ColorGradingSubpropertyConfig(
+    name: 'Contrast',
+    getDisplayName: (context) => AppLocalizations.of(context)!.colorGradingSubpropertyContrast,
+  ),
+  ColorGradingSubproperty.offset: ColorGradingSubpropertyConfig(
+    name: 'Offset',
+    getDisplayName: (context) => AppLocalizations.of(context)!.colorGradingSubpropertyOffset,
+    minValue: -1.0,
+    maxValue: 1.0,
+    saturationExponent: 3.0,
+  ),
+};
+
+/// Maximum height of the color wheel relative to the total height of the screen.
+const double maxWheelRadiusFromScreenHeight = 0.35;
+
+/// Get the localized display name corresponding to a color grading range.
+String getColorGradingRangeDisplayName(ColorGradingRange range, BuildContext context) {
+  final AppLocalizations localizations = AppLocalizations.of(context)!;
+
+  switch (range) {
+    case ColorGradingRange.global:
+      return localizations.colorGradingRangeGlobal;
+
+    case ColorGradingRange.shadows:
+      return localizations.colorGradingRangeShadows;
+
+    case ColorGradingRange.midtones:
+      return localizations.colorGradingRangeMidtones;
+
+    case ColorGradingRange.highlights:
+      return localizations.colorGradingRangeHighlights;
+  }
+}
+
+/// Get the configuration data corresponding to a color grading range.
+ColorGradingSubpropertyConfig? getColorGradingSubpropertyConfig(ColorGradingSubproperty subproperty) {
+  return _colorGradingSubpropertyConfig[subproperty];
+}
+
+/// Base widget for tabs that let users change color properties.
+class BaseColorTab extends StatefulWidget {
+  const BaseColorTab({
+    Key? key,
+    required this.colorProperties,
+    required this.mode,
+    required this.rightSideHeader,
+    required this.rightSideContents,
+    this.miscPropertyPrefix = '',
+    this.bUseEnableProperties = false,
+  }) : super(key: key);
+
+  /// Which mode to operate in. [colorProperties] must match the appropriate type.
+  final BaseColorTabMode mode;
+
+  /// Properties for the color values to be controlled.
+  final List<UnrealProperty> colorProperties;
+
+  /// Header for the right side column of the tab.
+  final Widget rightSideHeader;
+
+  /// Contents of the right side column of the tab.
+  final Widget rightSideContents;
+
+  /// Prefix for miscellaneous properties of the color grading settings.
+  final String miscPropertyPrefix;
+
+  /// Whether to control the enable/override properties for each property in the tab.
+  final bool bUseEnableProperties;
+
+  @override
+  State<BaseColorTab> createState() => _BaseColorTabState();
+}
+
+class _BaseColorTabState extends State<BaseColorTab> {
+  late final BaseColorTabSettings _tabSettings;
+
+  @override
+  void initState() {
+    super.initState();
+    _tabSettings = BaseColorTabSettings(PreferencesBundle.of(context));
+  }
+
+  Widget build(BuildContext context) {
+    final maxWheelRadius = MediaQuery.of(context).size.height * maxWheelRadiusFromScreenHeight;
+
+    late final Widget colorWheel; // The widget used to control the color property.
+    late final Widget miscPropertyControls; // Extra controls shown under the color wheel.
+
+    switch (widget.mode) {
+      case BaseColorTabMode.color:
+        colorWheel = UnrealColorWheel(
+          unrealProperties: widget.colorProperties,
+          bShowAlpha: false,
+          maxRadius: maxWheelRadius,
+        );
+        miscPropertyControls = const SizedBox();
+        break;
+
+      case BaseColorTabMode.colorGrading:
+      case BaseColorTabMode.postProcess:
+        colorWheel = _buildColorGradingWheel(maxWheelRadius);
+        miscPropertyControls = _buildColorGradingMiscControlsWidget();
+        break;
+
+      default:
+        throw 'Invalid actor type for top widget in details panel';
+    }
+
+    return Provider<BaseColorTabSettings>(
+      create: (_) => _tabSettings,
+      child: Padding(
+        padding: const EdgeInsets.only(top: sectionMargin),
+        child: Row(
+          crossAxisAlignment: CrossAxisAlignment.end,
+          children: [
+            Expanded(
+              child: Container(
+                color: Theme.of(context).colorScheme.surface,
+                child: Column(children: [
+                  colorWheel,
+                  Padding(
+                    padding: EdgeInsets.symmetric(horizontal: cardMargin),
+                    child: miscPropertyControls,
+                  ),
+                ]),
+              ),
+            ),
+
+            const SizedBox(width: sectionMargin),
+
+            // Other properties
+            Expanded(
+              child: Container(
+                color: Theme.of(context).colorScheme.surface,
+                child: Column(children: [
+                  Container(
+                    color: Theme.of(context).colorScheme.surfaceTint,
+                    height: 52,
+                    child: widget.rightSideHeader,
+                  ),
+                  Expanded(
+                    child: Container(
+                      padding: EdgeInsets.only(left: cardMargin, right: cardMargin, top: 16),
+                      color: Theme.of(context).colorScheme.surface,
+                      child: widget.rightSideContents,
+                    ),
+                  ),
+                ]),
+              ),
+            ),
+          ],
+        ),
+      ),
+    );
+  }
+
+  /// Get the display name of a range type.
+  String _getHighlightsRangeDisplayName(RangeLimitType rangeLimitType) {
+    switch (rangeLimitType) {
+      case RangeLimitType.max:
+        return AppLocalizations.of(context)!.propertyColorGradingHighlightsMax;
+
+      case RangeLimitType.min:
+        return AppLocalizations.of(context)!.propertyColorGradingHighlightsMin;
+    }
+  }
+
+  /// Get the property name of a range type.
+  String? _getHighlightsRangePropertyName(RangeLimitType rangeLimitType) {
+    String? rangeLimitName = _highlightRangeLimitPropertyNames[rangeLimitType];
+    if (rangeLimitName == null) {
+      return null;
+    }
+
+    return 'Highlights$rangeLimitName';
+  }
+
+  /// Build the wheel used to edit color grading values.
+  Widget _buildColorGradingWheel(double maxRadius) {
+    return PreferenceBuilder(
+      preference: _tabSettings.colorGradingRange,
+      builder: (context, ColorGradingRange colorGradingRange) {
+        final String rangeName = _colorGradingRangePropertyNames[colorGradingRange]!;
+
+        return PreferenceBuilder(
+          preference: _tabSettings.colorGradingSubproperty,
+          builder: (context, ColorGradingSubproperty colorGradingSubproperty) {
+            final ColorGradingSubpropertyConfig subpropertyConfig =
+                getColorGradingSubpropertyConfig(colorGradingSubproperty)!;
+
+            late final List<UnrealProperty> subproperties;
+            List<UnrealProperty>? enableProperties;
+
+            switch (widget.mode) {
+              case BaseColorTabMode.postProcess:
+                // For post-processing, "Global" is implied by no range prefix
+                final String rangePropertyName =
+                    (_tabSettings.colorGradingRange.getValue() == ColorGradingRange.global) ? '' : rangeName;
+
+                subproperties = getSubproperties(
+                  widget.colorProperties,
+                  'Color${subpropertyConfig.name}$rangePropertyName',
+                  typeNameOverride: 'FVector4_Color',
+                );
+
+                if (widget.bUseEnableProperties) {
+                  enableProperties = getSubproperties(
+                    widget.colorProperties,
+                    'bOverride_Color${subpropertyConfig.name}$rangePropertyName',
+                  );
+                }
+                break;
+
+              default:
+                subproperties = getSubproperties(
+                  widget.colorProperties,
+                  '$rangeName.${subpropertyConfig.name}',
+                  typeNameOverride: 'FVector4_Color',
+                );
+
+                if (widget.bUseEnableProperties) {
+                  enableProperties = getSubproperties(
+                    widget.colorProperties,
+                    '$rangeName.bOverride_${subpropertyConfig.name}',
+                  );
+                }
+                break;
+            }
+
+            return UnrealColorGradingWheel(
+              unrealProperties: subproperties,
+              enableProperties: enableProperties,
+              minValue: subpropertyConfig.minValue,
+              maxValue: subpropertyConfig.maxValue,
+              saturationExponent: subpropertyConfig.saturationExponent,
+              extraTopWidget: SizedBox(
+                width: 140,
+                child: DropdownSelector(
+                  value: colorGradingSubproperty,
+                  items: ColorGradingSubproperty.values,
+                  makeItemName: (value) =>
+                      (value != null
+                          ? getColorGradingSubpropertyConfig(value as ColorGradingSubproperty)?.getDisplayName(context)
+                          : null) ??
+                      '???',
+                  onChanged: (ColorGradingSubproperty value) => _tabSettings.colorGradingSubproperty.setValue(value),
+                ),
+              ),
+              extraBottomWidget: const _ColorGradingRangeToggleButtons(),
+              maxRadius: maxRadius,
+            );
+          },
+        );
+      },
+    );
+  }
+
+  /// Build the widgets to show under the wheel when editing color grading values.
+  Widget _buildColorGradingMiscControlsWidget() {
+    return PreferenceBuilder(
+      preference: _tabSettings.colorGradingRange,
+      builder: (context, ColorGradingRange colorGradingRange) {
+        final List<Widget> widgets = [];
+
+        // Add widgets under the wheel based on the grading range
+        switch (colorGradingRange) {
+          case ColorGradingRange.shadows:
+            widgets.add(UnrealDeltaSlider(
+              overrideName: AppLocalizations.of(context)!.propertyColorGradingShadowsMax,
+              unrealProperties: getSubproperties(
+                widget.colorProperties,
+                '${widget.miscPropertyPrefix}ShadowsMax',
+              ),
+              enableProperties: widget.bUseEnableProperties
+                  ? getSubproperties(
+                      widget.colorProperties,
+                      'bOverride_${widget.miscPropertyPrefix}ShadowsMax',
+                    )
+                  : null,
+            ));
+            break;
+
+          case ColorGradingRange.highlights:
+            widgets.add(
+              PreferenceBuilder(
+                preference: _tabSettings.colorGradingHighlightsRangeLimitType,
+                builder: (context, final RangeLimitType highlightsRangeLimitType) {
+                  final String? highlightsPropertyName = _getHighlightsRangePropertyName(highlightsRangeLimitType);
+
+                  if (highlightsPropertyName == null) {
+                    return const SizedBox();
+                  }
+
+                  return UnrealDeltaSlider(
+                    overrideName: _getHighlightsRangeDisplayName(highlightsRangeLimitType),
+                    buildLabel: (String name) => DrivenDropdownText(
+                      value: highlightsRangeLimitType,
+                      items: RangeLimitType.values,
+                      makeItemName: (RangeLimitType item) => _getHighlightsRangeDisplayName(item),
+                      onChanged: (RangeLimitType value) =>
+                          _tabSettings.colorGradingHighlightsRangeLimitType.setValue(value),
+                    ),
+                    unrealProperties: getSubproperties(
+                      widget.colorProperties,
+                      '${widget.miscPropertyPrefix}$highlightsPropertyName',
+                    ),
+                    enableProperties: widget.bUseEnableProperties
+                        ? getSubproperties(
+                            widget.colorProperties,
+                            'bOverride_${widget.miscPropertyPrefix}$highlightsPropertyName',
+                          )
+                        : null,
+                  );
+                },
+              ),
+            );
+            break;
+
+          default:
+            break;
+        }
+
+        return Column(children: widgets);
+      },
+    );
+  }
+}
+
+/// Toggle button selector for the color grading range to control.
+class _ColorGradingRangeToggleButtons extends StatelessWidget {
+  const _ColorGradingRangeToggleButtons({Key? key}) : super(key: key);
+
+  @override
+  Widget build(BuildContext context) {
+    final BaseColorTabSettings colorTabSettings = Provider.of<BaseColorTabSettings>(context, listen: false);
+
+    return PreferenceBuilder(
+      preference: Provider.of<BaseColorTabSettings>(context, listen: false).colorGradingRange,
+      builder: (context, ColorGradingRange colorGradingRange) => SelectorBar(
+        value: colorGradingRange,
+        onSelected: (newValue) => colorTabSettings.colorGradingRange.setValue(newValue as ColorGradingRange),
+        valueNames: {
+          for (final range in ColorGradingRange.values) range: getColorGradingRangeDisplayName(range, context),
+        },
+      ),
+    );
+  }
+}
