Branch snapshot for CL 27138787

diff --git a/Engine/Extras/VirtualProduction/EpicStageApp/lib/widgets/screens/main/tabs/color_grading_tab.dart b/Engine/Extras/VirtualProduction/EpicStageApp/lib/widgets/screens/main/tabs/color_grading_tab.dart
new file mode 100644
--- /dev/null
+++ b/Engine/Extras/VirtualProduction/EpicStageApp/lib/widgets/screens/main/tabs/color_grading_tab.dart
@@ -0,0 +1,1476 @@
+// Copyright Epic Games, Inc. All Rights Reserved.
+
+import 'dart:async';
+import 'dart:math' as math;
+
+import 'package:flutter/material.dart';
+import 'package:flutter_gen/gen_l10n/app_localizations.dart';
+import 'package:logging/logging.dart';
+import 'package:provider/provider.dart';
+import 'package:streaming_shared_preferences/streaming_shared_preferences.dart';
+
+import '../../../../models/engine_connection.dart';
+import '../../../../models/property_modify_operations.dart';
+import '../../../../models/settings/color_grading_tab_settings.dart';
+import '../../../../models/settings/delta_widget_settings.dart';
+import '../../../../models/settings/main_screen_settings.dart';
+import '../../../../models/settings/selected_actor_settings.dart';
+import '../../../../models/unreal_actor_manager.dart';
+import '../../../../models/unreal_transaction_manager.dart';
+import '../../../../models/unreal_types.dart';
+import '../../../../utilities/constants.dart';
+import '../../../../utilities/preferences_bundle.dart';
+import '../../../../utilities/transient_preference.dart';
+import '../../../../utilities/unreal_colors.dart';
+import '../../../../utilities/unreal_utilities.dart';
+import '../../../elements/delta_slider.dart';
+import '../../../elements/dropdown_text.dart';
+import '../../../elements/empty_placeholder.dart';
+import '../../../elements/epic_list_view.dart';
+import '../../../elements/layout/card.dart';
+import '../../../elements/modal.dart';
+import '../../../elements/reset_mode_button.dart';
+import '../../../elements/selector_bar.dart';
+import '../../../elements/stepper.dart';
+import '../../../elements/swipe_revealer.dart';
+import '../../../elements/tree_view.dart';
+import '../../../elements/unreal_property_builder.dart';
+import '../sidebar/outliner_panel.dart';
+import 'base_color_tab.dart';
+
+final _log = Logger('ColorGradingTab');
+
+/// Tab to control color grading settings on anything other than CCW/CCRs like walls, viewports, etc.
+class ColorGradingTab extends StatefulWidget {
+  const ColorGradingTab({Key? key}) : super(key: key);
+
+  static const String iconPath = 'assets/images/icons/color_grading.svg';
+
+  static String getTitle(BuildContext context) => AppLocalizations.of(context)!.tabTitleColorGrading;
+
+  @override
+  State<StatefulWidget> createState() => _ColorGradingTabState();
+}
+
+class _ColorGradingTabState extends State<ColorGradingTab> {
+  /// Minimum space to leave for the top panel when resizing the bottom one.
+  static const double _minTopPanelHeight = 108;
+
+  static const Duration _refreshRate = Duration(seconds: 3);
+
+  late final EngineConnectionManager _connectionManager;
+  late final UnrealActorManager _actorManager;
+  late final ColorGradingTabSettings _tabSettings;
+  late final DeltaWidgetSettings _deltaWidgetSettings;
+  late final StreamSubscription _targetSubscription;
+
+  final List<_ColorGradingObjectEntryData> _entries = [];
+  Timer? _refreshTimer;
+  Future? _bPendingRefresh = null;
+
+  @override
+  void initState() {
+    super.initState();
+
+    _connectionManager = Provider.of<EngineConnectionManager>(context, listen: false);
+    _actorManager = Provider.of<UnrealActorManager>(context, listen: false);
+
+    _actorManager.watchClassName(postProcessVolumeClassName, _actorUpdateStub);
+
+    _deltaWidgetSettings = Provider.of<DeltaWidgetSettings>(context, listen: false);
+
+    _tabSettings = ColorGradingTabSettings(PreferencesBundle.of(context));
+    _targetSubscription = _tabSettings.target.listen(_onTargetChanged);
+
+    _refreshTimer = Timer.periodic(_refreshRate, (timer) => _startRefreshingTargetList());
+    _startRefreshingTargetList();
+  }
+
+  @override
+  void dispose() {
+    _refreshTimer?.cancel();
+    _targetSubscription.cancel();
+
+    _actorManager.stopWatchingClassName(postProcessVolumeClassName, _actorUpdateStub);
+
+    _deltaWidgetSettings.removeResetModeBlocker(this);
+
+    super.dispose();
+  }
+
+  @override
+  Widget build(BuildContext context) {
+    return Provider<ColorGradingTabSettings>(
+      create: (_) => _tabSettings,
+      child: LayoutBuilder(builder: (context, constraints) {
+        // Determine how much space the targets panel is allowed to take up
+        double maxTargetPanelHeight = constraints.maxHeight - _minTopPanelHeight;
+
+        return TransientPreferenceBuilder(
+          preference: _tabSettings.target,
+          builder: (context, final UnrealProperty? colorGradingTarget) {
+            final _ColorGradingTargetListEntryData? targetData = _getTargetEntryDataForProperty(colorGradingTarget);
+
+            return Padding(
+              padding: EdgeInsets.all(cardMargin),
+              child: Row(children: [
+                // Main controls
+                Expanded(
+                  child: Card(
+                    child: Column(
+                      children: [
+                        CardLargeHeader(
+                          iconPath: _getIcon(),
+                          title: _getTitle(),
+                          subtitle: AppLocalizations.of(context)!.tabTitleColorGrading,
+                          trailing: const ResetModeButton(),
+                        ),
+                        Expanded(
+                          child: _ColorGradingMainControls(
+                            targetProperty: colorGradingTarget,
+                            type: _getEntryDataForSelectedObject()?.type,
+                          ),
+                        ),
+                      ],
+                    ),
+                  ),
+                ),
+
+                PreferenceBuilder(
+                  preference: Provider.of<MainScreenSettings>(context).bIsOutlinerPanelOpen,
+                  builder: (context, final bool bIsOutlinerPanelOpen) {
+                    if (!bIsOutlinerPanelOpen) {
+                      return const SizedBox();
+                    }
+
+                    return Row(children: [
+                      SizedBox(width: cardMargin),
+
+                      // Outliner and target panels
+                      SizedBox(
+                        width: getOutlinerWidth(context),
+                        child: Card(
+                          child: Column(
+                            children: [
+                              _ColorGradingOutlinerPanel(
+                                entries: _entries,
+                                onSelectionChanged: _onColorGradingObjectPathSelected,
+                                selectedObjectPath: colorGradingTarget?.objectPath,
+                              ),
+                              _ColorGradingTargetPanel(
+                                objectEntry: _getEntryDataForSelectedObject(),
+                                onTargetChanged: _tabSettings.target.setValue,
+                                currentTarget: targetData,
+                                maxHeight: maxTargetPanelHeight,
+                                refreshTargetList: _startRefreshingTargetList,
+                              )
+                            ],
+                          ),
+                        ),
+                      ),
+                    ]);
+                  },
+                ),
+              ]),
+            );
+          },
+        );
+      }),
+    );
+  }
+
+  /// Get the title to display at the top of the tab.
+  String? _getTitle() {
+    // Try to get the name of the current target
+    final _ColorGradingObjectEntryData? entryData = _getEntryDataForSelectedObject();
+    if (entryData != null) {
+      String? title = _getEntryDataForSelectedObject()?.name;
+
+      // No specific target name for PPVs
+      if (entryData.type != _ColorGradingObjectEntryType.postProcessVolume) {
+        // Get the stored name of the specific target
+        for (final _ColorGradingTargetListEntryData targetData in entryData.targets) {
+          if (targetData.property == _tabSettings.target.getValue()) {
+            title = '$title – ${targetData.name}';
+            break;
+          }
+        }
+      }
+
+      return title;
+    }
+
+    return null;
+  }
+
+  /// Get the path for the icon to display at the top of the tab.
+  String _getIcon() {
+    final _ColorGradingObjectEntryData? entryData = _getEntryDataForSelectedObject();
+
+    switch (entryData?.type) {
+      case _ColorGradingObjectEntryType.nDisplayConfig:
+      case _ColorGradingObjectEntryType.icvfxCamera:
+        return 'assets/images/icons/viewport.svg';
+
+      case _ColorGradingObjectEntryType.postProcessVolume:
+        return 'assets/images/icons/post_process_volume.svg';
+
+      default:
+        return 'assets/images/icons/color_grading.svg';
+    }
+  }
+
+  /// Start refreshing the target list if there isn't one in progress.
+  Future _startRefreshingTargetList() {
+    if (_bPendingRefresh == null) {
+      _bPendingRefresh = _refreshTargetList().then((_) => _bPendingRefresh = null);
+    }
+
+    return _bPendingRefresh!;
+  }
+
+  /// Refresh the list of possible color grading targets.
+  Future<void> _refreshTargetList() async {
+    final bool bHadTargetData = _getTargetEntryDataForProperty(_tabSettings.target.getValue()) != null;
+
+    _entries.clear();
+
+    await _addTargetForPostProcessVolumes();
+
+    final List<Future<void>> asyncJobs = [];
+
+    final selectedActorSettings = Provider.of<SelectedActorSettings>(context, listen: false);
+    final String displayClusterRootPath = selectedActorSettings.displayClusterRootPath.getValue();
+
+    final UnrealObject? rootActor = _actorManager.getActorAtPath(displayClusterRootPath);
+    late final Future<_ColorGradingObjectEntryData?> rootActorEntryData;
+
+    if (displayClusterRootPath.isNotEmpty && rootActor != null) {
+      rootActorEntryData = _addTargetsForConfigActorPath(rootActor);
+      asyncJobs.add(rootActorEntryData);
+    } else {
+      rootActorEntryData = Future.value(null);
+    }
+
+    await Future.wait(asyncJobs);
+
+    if (!mounted) {
+      return;
+    }
+
+    UnrealProperty colorGradingTarget = _tabSettings.target.getValue();
+
+    // If the current target isn't valid anymore, deselect it
+    final bool bIsValidTarget = colorGradingTarget.bIsNotEmpty &&
+        _entries.any((entryData) {
+          // For PPVs, we don't have a target within the object
+          if (entryData.type == _ColorGradingObjectEntryType.postProcessVolume &&
+              colorGradingTarget.objectPath == entryData.path) {
+            return true;
+          }
+
+          // Otherwise, check that the target property is still valid
+          return entryData.targets.any((targetData) {
+            return targetData.property == colorGradingTarget;
+          });
+        });
+
+    if (!bIsValidTarget) {
+      // Default to the root actor's entire cluster settings
+      rootActorEntryData.then((final _ColorGradingObjectEntryData? data) {
+        if (data == null) {
+          return;
+        }
+
+        final List<_ColorGradingTargetListEntryData> dataEntries = data.targets;
+        if (dataEntries.length < 1) {
+          return;
+        }
+
+        // First item in the entries list for the root actor will always be its entire cluster property
+        _tabSettings.target.setValue(dataEntries[0].property);
+      });
+
+      // If we still don't have a valid target, fall back to no target so we aren't pointing at something invalid
+      if (!_tabSettings.target.getValue().bIsNotEmpty) {
+        _tabSettings.target.setValue(UnrealProperty.empty);
+      }
+    }
+
+    if (!bHadTargetData) {
+      // If we didn't have valid target data before, we may now have retrieved the data that we needed.
+      // Normally this event will automatically fire when the target changes, but in this case, the target stays the
+      // same; we just didn't have its associated data ready yet.
+      _onTargetChanged(_tabSettings.target.getValue());
+    }
+
+    setState(() {});
+  }
+
+  /// Request the name of the current level and add post-process volumes under an entry for it.
+  Future<void> _addTargetForPostProcessVolumes() async {
+    final Set<UnrealObject> postProcessVolumes =
+        await _actorManager.getInitialActorsOfClass(postProcessVolumeClassName);
+
+    if (postProcessVolumes.length == 0) {
+      return;
+    }
+
+    String levelPath = postProcessVolumes.first.path;
+    levelPath = levelPath.substring(0, levelPath.indexOf(':'));
+
+    final String levelName = levelPath.substring(levelPath.indexOf('.') + 1);
+
+    // Add level
+    final levelEntry = _ColorGradingObjectEntryData(
+      path: levelPath,
+      name: levelName,
+      type: _ColorGradingObjectEntryType.level,
+      targets: [],
+      targetListTitle: '',
+    );
+    _entries.add(levelEntry);
+
+    // Add PPVs, which we already have a list of
+    for (final UnrealObject actor in postProcessVolumes) {
+      _entries.add(_ColorGradingObjectEntryData(
+        path: actor.path,
+        name: actor.name,
+        type: _ColorGradingObjectEntryType.postProcessVolume,
+        targets: [],
+        targetListTitle: AppLocalizations.of(context)!.colorGradingOutlinerTargetListTitleGeneric,
+        parent: levelEntry,
+      ));
+    }
+  }
+
+  /// Request data about an nDisplay config and add its color grading targets to the target list.
+  Future<_ColorGradingObjectEntryData?> _addTargetsForConfigActorPath(UnrealObject configActor) async {
+    // Find the path of the current nDisplay cluster's config data, as well as all of its components in case any of
+    // them are cameras.
+    final List<UnrealHttpRequest> requests = [
+      UnrealHttpRequest(url: '/remote/object/property', verb: 'PUT', body: {
+        'objectPath': configActor.path,
+        'propertyName': 'CurrentConfigData',
+        'access': 'READ_ACCESS',
+      }),
+      UnrealHttpRequest(url: '/remote/object/property', verb: 'PUT', body: {
+        'objectPath': configActor.path,
+        'propertyName': 'BlueprintCreatedComponents',
+        'access': 'READ_ACCESS',
+      }),
+    ];
+
+    final UnrealHttpResponse batchResponse = await _connectionManager.sendBatchedHttpRequest(requests);
+    if (batchResponse.code != 200) {
+      return null;
+    }
+
+    final _ColorGradingObjectEntryData? objectEntry =
+        await _addTargetsForConfigResponse(batchResponse.body[0], configActor);
+
+    if (objectEntry != null) {
+      await _addTargetsForComponentsResponse(batchResponse.body[1], objectEntry);
+    }
+
+    return objectEntry;
+  }
+
+  /// Given a response from querying an nDisplay root actor's configuration data, add all of its color grading targets
+  /// to the target list.
+  Future<_ColorGradingObjectEntryData?> _addTargetsForConfigResponse(
+      UnrealHttpResponse? response, UnrealObject configActor) async {
+    if (response?.code != 200 || !mounted) {
+      return null;
+    }
+
+    final String? configPath = response!.body?['CurrentConfigData'];
+    if (configPath == null) {
+      return null;
+    }
+
+    final List<_ColorGradingTargetListEntryData> targets = [];
+
+    // Add the config's entire cluster color grading
+    final entireClusterRootProperty = UnrealProperty(
+      objectPath: configPath,
+      propertyName: 'StageSettings.EntireClusterColorGrading',
+    );
+
+    targets.add(
+      _ColorGradingTargetListEntryData(
+        name: AppLocalizations.of(context)!.colorGradingOutlinerTargetEntireCluster,
+        bIsListEntry: false,
+        property: entireClusterRootProperty.makeSubproperty('ColorGradingSettings'),
+        enableProperty: entireClusterRootProperty.makeSubproperty('bEnableEntireClusterColorGrading'),
+      ),
+    );
+
+    // Get the config's per-viewport color grading settings
+    final perViewportRequest = UnrealHttpRequest(url: '/remote/object/property', verb: 'PUT', body: {
+      'objectPath': configPath,
+      'propertyName': 'StageSettings.PerViewportColorGrading',
+      'access': 'READ_ACCESS',
+    });
+
+    final UnrealHttpResponse perViewportResponse = await _connectionManager.sendHttpRequest(perViewportRequest);
+    if (perViewportResponse.code != 200 || !mounted) {
+      return null;
+    }
+
+    final List<dynamic>? perViewportTargets = perViewportResponse.body['PerViewportColorGrading'];
+    if (perViewportTargets == null) {
+      return null;
+    }
+
+    for (int entryIndex = 0; entryIndex < perViewportTargets.length; ++entryIndex) {
+      final dynamic entry = perViewportTargets[entryIndex];
+      final String? entryName = entry['Name'];
+      if (entryName == null) {
+        continue;
+      }
+
+      final targetRootProperty = UnrealProperty(
+        objectPath: configPath,
+        propertyName: 'StageSettings.PerViewportColorGrading[$entryIndex]',
+      );
+
+      targets.add(_ColorGradingTargetListEntryData(
+        name: entryName,
+        bIsListEntry: true,
+        property: targetRootProperty.makeSubproperty('ColorGradingSettings'),
+        enableProperty: targetRootProperty.makeSubproperty('bIsEnabled'),
+        nameProperty: targetRootProperty.makeSubproperty('Name'),
+      ));
+    }
+
+    final objectEntry = _ColorGradingObjectEntryData(
+      path: configPath,
+      name: configActor.name,
+      type: _ColorGradingObjectEntryType.nDisplayConfig,
+      targetListTitle: AppLocalizations.of(context)!.colorGradingOutlinerTargetListTitlePerViewport,
+      targets: targets,
+    );
+    _entries.add(objectEntry);
+
+    return objectEntry;
+  }
+
+  /// Given a response from querying an nDisplay root actor's components list, find any ICVFX cameras and add their
+  /// color grading targets to the target list.
+  Future<void> _addTargetsForComponentsResponse(
+      UnrealHttpResponse? response, _ColorGradingObjectEntryData parent) async {
+    if (response?.code != 200 || !mounted) {
+      return;
+    }
+
+    final List<dynamic>? componentPaths = response!.body?['BlueprintCreatedComponents'];
+    if (componentPaths == null) {
+      return;
+    }
+
+    // Try to get the per-node color grading settings from each component as if it's an ICVFX camera. If it succeeds,
+    // we can treat it as one and add all its targets to the list.
+    final List<UnrealHttpRequest> describeRequests = [];
+    for (final dynamic componentPath in componentPaths) {
+      if (componentPath is! String) {
+        continue;
+      }
+
+      describeRequests.add(UnrealHttpRequest(url: '/remote/object/describe', verb: 'PUT', body: {
+        'objectPath': componentPath,
+        'access': 'READ_ACCESS',
+      }));
+    }
+
+    final UnrealHttpResponse batchDescribeResponse = await _connectionManager.sendBatchedHttpRequest(describeRequests);
+    if (batchDescribeResponse.code != 200 || !mounted) {
+      return;
+    }
+
+    // Find which of the described responses are actually cameras
+    final List<String> cameraPaths = [];
+    for (int componentIndex = 0; componentIndex < componentPaths.length; ++componentIndex) {
+      final UnrealHttpResponse? describeResponse = batchDescribeResponse.body[componentIndex];
+      if (describeResponse == null || describeResponse.code != 200) {
+        continue;
+      }
+
+      if (describeResponse.body?['Class'] == '/Script/DisplayCluster.DisplayClusterICVFXCameraComponent') {
+        cameraPaths.add(componentPaths[componentIndex]);
+      }
+    }
+
+    // Now that we know which are cameras, query them all for their color grading targets and add them to the list
+    final List<UnrealHttpRequest> propertyRequests = [];
+    for (final String cameraPath in cameraPaths) {
+      propertyRequests.add(UnrealHttpRequest(url: '/remote/object/property', verb: 'PUT', body: {
+        'objectPath': cameraPath,
+        'propertyName': 'CameraSettings.PerNodeColorGrading',
+        'access': 'READ_ACCESS',
+      }));
+    }
+
+    final UnrealHttpResponse batchPropertyResponse = await _connectionManager.sendBatchedHttpRequest(propertyRequests);
+    if (batchPropertyResponse.code != 200 || !mounted) {
+      return;
+    }
+
+    for (int cameraIndex = 0; cameraIndex < cameraPaths.length; ++cameraIndex) {
+      final UnrealHttpResponse? propertyResponse = batchPropertyResponse.body[cameraIndex];
+      _addTargetsForICVFXCameraResponse(cameraPaths[cameraIndex], propertyResponse, parent);
+    }
+  }
+
+  /// Given a response from querying an ICVFX camera's list of per-node color grading settings, add all of its color
+  /// grading targets to the target list.
+  void _addTargetsForICVFXCameraResponse(
+      String cameraPath, UnrealHttpResponse? response, _ColorGradingObjectEntryData parent) {
+    if (response?.code != 200) {
+      return;
+    }
+
+    final List<_ColorGradingTargetListEntryData> targets = [];
+
+    // All nodes color grading is always available
+    final allNodesRootProperty = UnrealProperty(
+      objectPath: cameraPath,
+      propertyName: 'CameraSettings.AllNodesColorGrading',
+    );
+
+    targets.add(_ColorGradingTargetListEntryData(
+      name: AppLocalizations.of(context)!.colorGradingOutlinerTargetAllNodes,
+      bIsListEntry: false,
+      property: allNodesRootProperty.makeSubproperty('ColorGradingSettings'),
+      enableProperty: allNodesRootProperty.makeSubproperty('bEnableInnerFrustumAllNodesColorGrading'),
+    ));
+
+    // Add targets for each per-node entry
+    final List<dynamic>? perNodeEntries = response!.body['PerNodeColorGrading'];
+    if (perNodeEntries == null) {
+      return;
+    }
+
+    for (int entryIndex = 0; entryIndex < perNodeEntries.length; ++entryIndex) {
+      final dynamic entry = perNodeEntries[entryIndex];
+      final String? entryName = entry['Name'];
+      if (entryName == null) {
+        continue;
+      }
+
+      final targetRootProperty = UnrealProperty(
+        objectPath: cameraPath,
+        propertyName: 'CameraSettings.PerNodeColorGrading[$entryIndex]',
+      );
+
+      targets.add(_ColorGradingTargetListEntryData(
+        name: entryName,
+        bIsListEntry: true,
+        property: targetRootProperty.makeSubproperty('ColorGradingSettings'),
+        enableProperty: targetRootProperty.makeSubproperty('bIsEnabled'),
+        nameProperty: targetRootProperty.makeSubproperty('Name'),
+      ));
+    }
+
+    final int lastDotIndex = cameraPath.lastIndexOf('.');
+    final String cameraName = cameraPath.substring(lastDotIndex + 1);
+
+    _entries.add(_ColorGradingObjectEntryData(
+      path: cameraPath,
+      name: cameraName,
+      type: _ColorGradingObjectEntryType.icvfxCamera,
+      targetListTitle: AppLocalizations.of(context)!.colorGradingOutlinerTargetListTitlePerNode,
+      targets: targets,
+      parent: parent,
+    ));
+  }
+
+  /// Get the data for the given object path.
+  _ColorGradingObjectEntryData? _getEntryDataForObject(String objectPath) {
+    for (final _ColorGradingObjectEntryData entryData in _entries) {
+      if (entryData.path == objectPath) {
+        return entryData;
+      }
+    }
+
+    return null;
+  }
+
+  /// Get the data for the currently selected entry.
+  _ColorGradingObjectEntryData? _getEntryDataForSelectedObject() {
+    final UnrealProperty? colorGradingTarget = _tabSettings.target.getValue();
+    if (colorGradingTarget == null) {
+      return null;
+    }
+
+    return _getEntryDataForObject(colorGradingTarget.objectPath);
+  }
+
+  /// Get the target data for the given property.
+  _ColorGradingTargetListEntryData? _getTargetEntryDataForProperty(UnrealProperty? targetProperty) {
+    if (targetProperty == null) {
+      return null;
+    }
+
+    for (final _ColorGradingObjectEntryData entryData in _entries) {
+      if (entryData.path == targetProperty.objectPath) {
+        for (final _ColorGradingTargetListEntryData target in entryData.targets) {
+          if (targetProperty == target.property) {
+            return target;
+          }
+        }
+      }
+    }
+
+    return null;
+  }
+
+  /// Called when the user selects a new color grading object.
+  void _onColorGradingObjectPathSelected(String? newSelectedPath) async {
+    // Wait for any in-progress refreshes to finish so we have entryData to work with.
+    if (_bPendingRefresh != null) {
+      await _bPendingRefresh;
+    }
+
+    if (newSelectedPath == _tabSettings.target.getValue().objectPath) {
+      return;
+    }
+
+    if (newSelectedPath == null) {
+      _tabSettings.target.setValue(UnrealProperty.empty);
+      return;
+    }
+
+    final _ColorGradingObjectEntryData? entryData = _getEntryDataForObject(newSelectedPath);
+    if (entryData == null) {
+      _tabSettings.target.setValue(UnrealProperty.empty);
+      return;
+    }
+
+    if (entryData.type == _ColorGradingObjectEntryType.postProcessVolume) {
+      // PPVs have only one possible target
+      _tabSettings.target.setValue(UnrealProperty(objectPath: newSelectedPath, propertyName: 'Settings'));
+      return;
+    }
+
+    // Select the first target, if available
+    _tabSettings.target.setValue(
+      (entryData.targets.length > 0)
+          ? entryData.targets.first.property
+          : UnrealProperty(objectPath: newSelectedPath, propertyName: ''),
+    );
+  }
+
+  /// Called when the color grading target changes.
+  void _onTargetChanged(UnrealProperty target) {
+    final engineConnection = Provider.of<EngineConnectionManager>(context, listen: false);
+    if (engineConnection.apiVersion?.bCanResetListItems != true) {
+      _blockResetModeForTarget(target);
+    }
+  }
+
+  /// Block reset mode if necessary for this [target] due to engine API limitations.
+  void _blockResetModeForTarget(UnrealProperty target) {
+    final _ColorGradingTargetListEntryData? targetData = _getTargetEntryDataForProperty(target);
+    if (targetData == null) {
+      return;
+    }
+
+    final bool bIsBlocked = _deltaWidgetSettings.isResetModeBlockedBy(this);
+    final bool bShouldBlock = targetData.bIsListEntry;
+
+    if (bShouldBlock && !bIsBlocked) {
+      _deltaWidgetSettings.addResetModeBlocker(
+        this,
+        AppLocalizations.of(context)!.resetModeColorGradingNotSupported,
+      );
+    } else if (!bShouldBlock && bIsBlocked) {
+      _deltaWidgetSettings.removeResetModeBlocker(this);
+    }
+  }
+
+  /// Callback function when actor details change. We handle this in _refreshTargetList instead, so we don't need any
+  /// implementation here.
+  void _actorUpdateStub(ActorUpdateDetails details) {}
+}
+
+/// The outliner panel for the color grading tab, showing the top-level actors or components that can be targeted.
+class _ColorGradingOutlinerPanel extends StatefulWidget {
+  const _ColorGradingOutlinerPanel({
+    Key? key,
+    required this.entries,
+    required this.selectedObjectPath,
+    required this.onSelectionChanged,
+  }) : super(key: key);
+
+  /// A flat list of all objects containing color grading targets.
+  final List<_ColorGradingObjectEntryData> entries;
+
+  /// The path of the currently selected object.
+  final String? selectedObjectPath;
+
+  /// A function called when the target object.
+  final Function(String? newSelectedObjectPath) onSelectionChanged;
+
+  @override
+  State<_ColorGradingOutlinerPanel> createState() => _ColorGradingOutlinerPanelState();
+}
+
+class _ColorGradingOutlinerPanelState extends State<_ColorGradingOutlinerPanel> {
+  late final TreeViewController _treeController = TreeViewController();
+
+  @override
+  void initState() {
+    super.initState();
+
+    _updateTreeController();
+  }
+
+  @override
+  void didUpdateWidget(covariant _ColorGradingOutlinerPanel oldWidget) {
+    super.didUpdateWidget(oldWidget);
+
+    _updateTreeController();
+  }
+
+  @override
+  Widget build(BuildContext context) {
+    return Flexible(
+      flex: 6,
+      child: Column(
+        children: [
+          CardSmallHeader(title: AppLocalizations.of(context)!.outlinerTitle),
+          Expanded(
+            child: TreeView(
+              padding: EdgeInsets.symmetric(vertical: 4),
+              treeController: _treeController,
+              nodeBuilder: (node, controller) => _ColorGradingObjectEntry(
+                node: node as TreeViewNode<_ColorGradingObjectEntryData>,
+                controller: _treeController,
+                onTap: () => widget.onSelectionChanged(node.data.path),
+                bIsSelected: node.data.path == widget.selectedObjectPath,
+              ),
+            ),
+          ),
+        ],
+      ),
+    );
+  }
+
+  /// Update the tree controller with current entries and remove old ones.
+  void _updateTreeController() {
+    final Set<String> relevantKeys = {};
+
+    // Add new entries
+    for (final _ColorGradingObjectEntryData entry in widget.entries) {
+      relevantKeys.add(entry.path);
+      if (_treeController.getNode(entry.path) == null) {
+        _treeController.addNode(_makeNode(entry), parentKey: entry.parent?.path);
+      }
+    }
+
+    // Remove old entries
+    final List<String> allKeys = List.from(_treeController.allKeys);
+    for (final String key in allKeys) {
+      if (!relevantKeys.contains(key)) {
+        _treeController.removeNode(key);
+      }
+    }
+  }
+
+  /// Make a [TreeViewNode] for a color grading object.
+  TreeViewNode<_ColorGradingObjectEntryData> _makeNode(final _ColorGradingObjectEntryData entry) {
+    return TreeViewNode(
+      key: entry.path,
+      data: entry,
+    );
+  }
+}
+
+/// The panel showing specific color grading settings that can be targeted on the selected actor.
+class _ColorGradingTargetPanel extends StatefulWidget {
+  const _ColorGradingTargetPanel({
+    Key? key,
+    required this.objectEntry,
+    required this.currentTarget,
+    required this.onTargetChanged,
+    required this.maxHeight,
+    required this.refreshTargetList,
+  }) : super(key: key);
+
+  /// The selected outliner entry containing the relevant targets, or null if one isn't selected.
+  final _ColorGradingObjectEntryData? objectEntry;
+
+  /// Data for the currently selected target color grading property.
+  final _ColorGradingTargetListEntryData? currentTarget;
+
+  /// A function called when the target property changes.
+  final Function(UnrealProperty newTarget) onTargetChanged;
+
+  /// A function called when the target list should refresh. Returns a Future which will complete when the list has been
+  /// refreshed.
+  final Future Function() refreshTargetList;
+
+  /// The maximum height of this panel, including its header.
+  final double maxHeight;
+
+  @override
+  State<_ColorGradingTargetPanel> createState() => _ColorGradingTargetPanelState();
+}
+
+class _ColorGradingTargetPanelState extends State<_ColorGradingTargetPanel> {
+  final GlobalKey _innerListKey = GlobalKey();
+
+  /// The user's touch position when they started dragging the title bar.
+  Offset? _dragStartPosition;
+
+  /// The height of the panel when the last drag operation started.
+  double _dragStartHeight = 0;
+
+  @override
+  Widget build(BuildContext context) {
+    final connection = Provider.of<EngineConnectionManager>(context, listen: false);
+
+    late final bCanSelectTarget;
+
+    // Determine whether we can select sub-targets within object entry
+    switch (widget.objectEntry?.type) {
+      case _ColorGradingObjectEntryType.nDisplayConfig:
+      case _ColorGradingObjectEntryType.icvfxCamera:
+        bCanSelectTarget = true;
+        break;
+
+      default:
+        bCanSelectTarget = false;
+        break;
+    }
+
+    late final Widget innerList;
+    if (!bCanSelectTarget) {
+      innerList = EmptyPlaceholder(
+        key: _innerListKey,
+        message: AppLocalizations.of(context)!.colorGradingOutlinerTargetListEmptyMessage,
+      );
+    } else {
+      final List<_ColorGradingTargetListEntryData> targets = widget.objectEntry!.targets;
+
+      innerList = EpicListView(
+        key: _innerListKey,
+        padding: EdgeInsets.only(top: 4),
+        itemCount: targets.length,
+        itemBuilder: (BuildContext context, int targetIndex) {
+          final _ColorGradingTargetListEntryData entryData = targets[targetIndex];
+          return _ColorGradingTargetListEntry(
+            data: entryData,
+            onTap: () => widget.onTargetChanged(entryData.property),
+            bIsSelected: entryData.property == widget.currentTarget?.property,
+          );
+        },
+      );
+    }
+
+    return PreferenceBuilder(
+      preference: Provider.of<ColorGradingTabSettings>(context, listen: false).panelHeight,
+      builder: (context, double panelHeight) {
+        final bool bHasSetHeight = panelHeight >= 0;
+
+        final mainColumn = Column(
+          mainAxisSize: MainAxisSize.min,
+          children: [
+            Container(
+              color: Theme.of(context).colorScheme.surfaceTint,
+              child: Column(
+                mainAxisSize: MainAxisSize.min,
+                children: [
+                  GestureDetector(
+                    child: CardSmallHeader(
+                      title: widget.objectEntry?.targetListTitle ??
+                          AppLocalizations.of(context)!.colorGradingOutlinerTargetListTitleGeneric,
+                    ),
+                    onVerticalDragStart: _onVerticalDragStart,
+                    onVerticalDragUpdate: _onVerticalDragUpdate,
+                    onVerticalDragEnd: _onVerticalDragEnd,
+                    onVerticalDragCancel: _onVerticalDragCancel,
+                  ),
+                  CardSubHeader(
+                    padding: EdgeInsets.symmetric(horizontal: 4),
+                    child: Row(children: [
+                      CardSubHeaderButton(
+                        iconPath: 'assets/images/icons/add_circle.svg',
+                        tooltipMessage: AppLocalizations.of(context)!.colorGradingOutlinerAddTargetButtonTooltip,
+                        onPressed: widget.objectEntry?.targetListProperty != null ? _addTargetToSelectedObject : null,
+                      ),
+                      CardSubHeaderButton(
+                        iconPath: 'assets/images/icons/edit.svg',
+                        tooltipMessage: AppLocalizations.of(context)!.colorGradingOutlinerRenameTargetButtonTooltip,
+                        onPressed: widget.currentTarget?.nameProperty != null ? _renameTarget : null,
+                      ),
+                      const Spacer(),
+                      CardSubHeaderButton(
+                        iconPath: 'assets/images/icons/trash.svg',
+                        tooltipMessage: AppLocalizations.of(context)!.colorGradingOutlinerDeleteTargetButtonTooltip,
+                        onPressed: connection.apiVersion?.bCanUseQueryParamsInWebSocketHttpUrl == true &&
+                                widget.currentTarget?.bIsListEntry == true
+                            ? _deleteTarget
+                            : null,
+                      ),
+                    ]),
+                  ),
+                ],
+              ),
+            ),
+            if (bHasSetHeight)
+              ConstrainedBox(
+                constraints: BoxConstraints(
+                  maxHeight: math.min(panelHeight, widget.maxHeight),
+                ),
+                child: innerList,
+              )
+            else
+              Expanded(child: innerList)
+          ],
+        );
+
+        if (bHasSetHeight) {
+          return mainColumn;
+        }
+
+        return Flexible(
+          child: mainColumn,
+          flex: 5,
+        );
+      },
+    );
+  }
+
+  /// Called when the user starts to drag the title bar of this panel.
+  void _onVerticalDragStart(DragStartDetails details) {
+    _dragStartPosition = details.globalPosition;
+
+    final tabSettings = Provider.of<ColorGradingTabSettings>(context, listen: false);
+    final double panelHeight = tabSettings.panelHeight.getValue();
+
+    if (panelHeight < 0) {
+      // We don't have a set height, so calculate it from the render box
+      final innerRenderBox = _innerListKey.currentContext?.findRenderObject() as RenderBox?;
+      if (innerRenderBox == null) {
+        return;
+      }
+
+      _dragStartHeight = innerRenderBox.size.height;
+    } else {
+      _dragStartHeight = panelHeight;
+    }
+  }
+
+  /// Called when the user continues to drag the title bar of this panel.
+  void _onVerticalDragUpdate(DragUpdateDetails details) {
+    if (_dragStartPosition == null) {
+      return;
+    }
+
+    final double newHeight = _dragStartHeight + (_dragStartPosition!.dy - details.globalPosition.dy);
+
+    final tabSettings = Provider.of<ColorGradingTabSettings>(context, listen: false);
+    tabSettings.panelHeight.setValue(newHeight.clamp(0, widget.maxHeight));
+  }
+
+  /// Called when the user finishes dragging the title bar of this panel.
+  void _onVerticalDragEnd(DragEndDetails details) {
+    _onVerticalDragCancel();
+  }
+
+  /// Called when a drag gesture on the title bar is cancelled.
+  void _onVerticalDragCancel() {
+    _dragStartPosition = null;
+  }
+
+  /// Show a dialog to rename the selected color grading target.
+  void _renameTarget() {
+    if (widget.currentTarget == null) {
+      return;
+    }
+
+    final _ColorGradingTargetListEntryData target = widget.currentTarget!;
+
+    GenericModalDialogRoute.showDialog(
+      context: context,
+      builder: (_) => StringTextInputModalDialog(
+        title: AppLocalizations.of(context)!.renameColorGradingTargetTitle,
+        initialValue: target.name,
+        handleResult: (result) => _handleRenameTargetResult(result, target),
+      ),
+    );
+  }
+
+  /// Handle the result of the dialog shown by [_renameTarget] as applied to a [target].
+  Future _handleRenameTargetResult(
+      TextInputModalDialogResult<String> result, _ColorGradingTargetListEntryData target) async {
+    if (result.action != TextInputModalDialogAction.apply || target.nameProperty == null) {
+      return;
+    }
+
+    final transactionManager = Provider.of<UnrealTransactionManager>(context, listen: false);
+    final connection = Provider.of<EngineConnectionManager>(context, listen: false);
+
+    final Future<UnrealHttpResponse> messageFuture = transactionManager.wrapWithTransactionIfManualAllowedForProperties(
+      AppLocalizations.of(context)!.transactionRenameColorGradingTarget,
+      (bIsManualTransaction) => connection.sendHttpRequest(
+        UnrealHttpRequest(
+          url: '/remote/object/property',
+          verb: 'PUT',
+          body: {
+            'access': bIsManualTransaction ? 'WRITE_MANUAL_TRANSACTION_ACCESS' : 'WRITE_TRANSACTION_ACCESS',
+            'generateTransaction': !bIsManualTransaction,
+            'objectPath': target.nameProperty!.objectPath,
+            'propertyName': target.nameProperty!.propertyName,
+            'propertyValue': {
+              target.nameProperty!.lastPropertyName: result.value!,
+            },
+          },
+        ),
+      ),
+    );
+
+    final UnrealHttpResponse response = await messageFuture;
+    if (response.code == HttpResponseCode.ok) {
+      await widget.refreshTargetList();
+    } else {
+      _log.warning('Failed to rename color grading target "${target.name}" to "${result.value}"');
+    }
+  }
+
+  /// Show a dialog to add a color grading target to the selected object.
+  void _addTargetToSelectedObject() {
+    if (widget.objectEntry == null) {
+      return;
+    }
+
+    GenericModalDialogRoute.showDialog(
+      context: context,
+      builder: (_) => StringTextInputModalDialog(
+        title: AppLocalizations.of(context)!.addColorGradingTargetTitle,
+        initialValue: AppLocalizations.of(context)!.addColorGradingTargetDefaultName,
+        handleResult: (result) => _handleAddTargetResult(result, widget.objectEntry!),
+      ),
+    );
+  }
+
+  /// Handle the result of the dialog shown by [_addTargetToSelectedObject] as applied to a [target].
+  Future _handleAddTargetResult(TextInputModalDialogResult<String> result, _ColorGradingObjectEntryData target) async {
+    if (result.action != TextInputModalDialogAction.apply || target.targetListProperty == null) {
+      return;
+    }
+
+    final transactionManager = Provider.of<UnrealTransactionManager>(context, listen: false);
+    final connection = Provider.of<EngineConnectionManager>(context, listen: false);
+
+    final Future<UnrealHttpResponse> messageFuture = transactionManager.wrapWithTransactionIfManualAllowedForProperties(
+      AppLocalizations.of(context)!.transactionAddColorGradingTarget,
+      (bIsManualTransaction) => connection.sendHttpRequest(
+        UnrealHttpRequest(
+          url: '/remote/object/property/append',
+          verb: 'PUT',
+          body: {
+            'access': bIsManualTransaction ? 'WRITE_MANUAL_TRANSACTION_ACCESS' : 'WRITE_TRANSACTION_ACCESS',
+            'generateTransaction': !bIsManualTransaction,
+            'objectPath': target.path,
+            'propertyName': target.targetListProperty!.propertyName,
+            'propertyValue': {
+              target.targetListProperty!.lastPropertyName: {'Name': result.value!},
+            },
+          },
+        ),
+      ),
+    );
+
+    final UnrealHttpResponse response = await messageFuture;
+    if (response.code == HttpResponseCode.ok) {
+      await widget.refreshTargetList();
+    } else {
+      _log.warning('Failed to add color grading target to "${target.name}"');
+    }
+  }
+
+  /// Delete the selected color grading target.
+  void _deleteTarget() async {
+    if (widget.objectEntry == null || widget.objectEntry!.targetListProperty == null || widget.currentTarget == null) {
+      return;
+    }
+
+    final connection = Provider.of<EngineConnectionManager>(context, listen: false);
+    if (connection.apiVersion?.bCanUseQueryParamsInWebSocketHttpUrl != true) {
+      return;
+    }
+
+    // Subtract one to ignore the first entry, which is always the entire cluster/camera VFX and can't be removed
+    final int index = widget.objectEntry!.targets.indexOf(widget.currentTarget!) - 1;
+
+    final transactionManager = Provider.of<UnrealTransactionManager>(context, listen: false);
+
+    final Future<UnrealHttpResponse> messageFuture = transactionManager.wrapWithTransactionIfManualAllowedForProperties(
+      AppLocalizations.of(context)!.transactionAddColorGradingTarget,
+      (bIsManualTransaction) => connection.sendHttpRequest(
+        UnrealHttpRequest(
+          url: '/remote/object/property/remove?index=$index',
+          verb: 'PUT',
+          body: {
+            'access': bIsManualTransaction ? 'WRITE_MANUAL_TRANSACTION_ACCESS' : 'WRITE_TRANSACTION_ACCESS',
+            'generateTransaction': !bIsManualTransaction,
+            'objectPath': widget.objectEntry!.path,
+            'propertyName': widget.objectEntry!.targetListProperty!.propertyName,
+          },
+        ),
+      ),
+    );
+
+    final UnrealHttpResponse response = await messageFuture;
+    if (response.code == HttpResponseCode.ok) {
+      await widget.refreshTargetList();
+    } else {
+      _log.warning('Failed to delete grading target "${widget.currentTarget!.name}" (index $index) '
+          'from "${widget.objectEntry!.name}"');
+    }
+  }
+}
+
+/// The main controls of the color grading tab.
+class _ColorGradingMainControls extends StatelessWidget {
+  const _ColorGradingMainControls({
+    Key? key,
+    required this.targetProperty,
+    required this.type,
+  }) : super(key: key);
+
+  /// The base color grading property being controlled.
+  final UnrealProperty? targetProperty;
+
+  /// The type of object being controlled.
+  final _ColorGradingObjectEntryType? type;
+
+  @override
+  Widget build(BuildContext context) {
+    if (targetProperty == null || targetProperty!.propertyName == '' || type == null) {
+      return SizedBox();
+    }
+
+    late final String whiteBalancePrefix;
+    late final String miscPropertyPrefix;
+    late final BaseColorTabMode mode;
+
+    switch (type) {
+      case _ColorGradingObjectEntryType.postProcessVolume:
+        mode = BaseColorTabMode.postProcess;
+        whiteBalancePrefix = '';
+        miscPropertyPrefix = '';
+        break;
+
+      default:
+        mode = BaseColorTabMode.colorGrading;
+        whiteBalancePrefix = 'WhiteBalance.';
+        miscPropertyPrefix = 'Misc.';
+        break;
+    }
+
+    final List<UnrealProperty> temperatureTypeProperties =
+        getSubproperties([targetProperty!], '${whiteBalancePrefix}TemperatureType');
+
+    return BaseColorTab(
+      mode: mode,
+      colorProperties: [targetProperty!],
+      miscPropertyPrefix: 'ColorCorrection',
+      bUseEnableProperties: true,
+      rightSideHeader: Center(child: FakeSelectorBar(AppLocalizations.of(context)!.colorTabAppearancePropertiesLabel)),
+      rightSideContents: Column(
+        crossAxisAlignment: CrossAxisAlignment.center,
+        children: [
+          UnrealMultiPropertyBuilder<String>(
+            properties: temperatureTypeProperties,
+            fallbackValue: AppLocalizations.of(context)!.mismatchedValuesLabel,
+            builder: (context, String? sharedValue, _) => UnrealDeltaSlider(
+              overrideName: sharedValue,
+              unrealProperties: getSubproperties([targetProperty!], '${whiteBalancePrefix}WhiteTemp'),
+              enableProperties: getSubproperties([targetProperty!], '${whiteBalancePrefix}bOverride_WhiteTemp'),
+              buildLabel: (name) => UnrealDropdownText(
+                unrealProperties: temperatureTypeProperties,
+                enableProperties: getSubproperties([targetProperty!], '${whiteBalancePrefix}bOverride_TemperatureType'),
+              ),
+            ),
+          ),
+          UnrealDeltaSlider(
+            unrealProperties: getSubproperties([targetProperty!], '${whiteBalancePrefix}WhiteTint'),
+            enableProperties: getSubproperties([targetProperty!], '${whiteBalancePrefix}bOverride_WhiteTint'),
+            overrideName: AppLocalizations.of(context)!.propertyColorGradingTint,
+          ),
+          UnrealDeltaSlider(
+            unrealProperties: getSubproperties([targetProperty!], '${miscPropertyPrefix}BlueCorrection'),
+            enableProperties: getSubproperties([targetProperty!], '${miscPropertyPrefix}bOverride_BlueCorrection'),
+          ),
+          UnrealDeltaSlider(
+            unrealProperties: getSubproperties([targetProperty!], '${miscPropertyPrefix}ExpandGamut'),
+            enableProperties: getSubproperties([targetProperty!], '${miscPropertyPrefix}bOverride_ExpandGamut'),
+          ),
+          UnrealDeltaSlider(
+            unrealProperties: getSubproperties([targetProperty!], 'AutoExposureBias'),
+            enableProperties: getSubproperties([targetProperty!], 'bOverride_AutoExposureBias'),
+          ),
+          UnrealStepper(
+            unrealProperties: getSubproperties([targetProperty!], 'AutoExposureBias'),
+            enableProperties: getSubproperties([targetProperty!], 'bOverride_AutoExposureBias'),
+            steps: StepperStepConfig.exposureSteps,
+          ),
+        ],
+      ),
+    );
+  }
+}
+
+/// Widget representing an object listed in the color grading outliner panel.
+class _ColorGradingObjectEntry extends StatelessWidget {
+  const _ColorGradingObjectEntry({
+    required this.node,
+    required this.controller,
+    required this.onTap,
+    required this.bIsSelected,
+    Key? key,
+  }) : super(key: key);
+
+  static const Map<_ColorGradingObjectEntryType, String> _iconsByType = {
+    _ColorGradingObjectEntryType.level: 'assets/images/icons/level.svg',
+    _ColorGradingObjectEntryType.nDisplayConfig: 'assets/images/icons/ndisplay.svg',
+    _ColorGradingObjectEntryType.icvfxCamera: 'assets/images/icons/ndisplay_camera.svg',
+    _ColorGradingObjectEntryType.postProcessVolume: 'assets/images/icons/post_process_volume.svg',
+  };
+
+  /// The node in the tree view for this item.
+  final TreeViewNode<_ColorGradingObjectEntryData> node;
+
+  /// Tree view controller for the containing tree.
+  final TreeViewController controller;
+
+  /// Callback for when this item is tapped.
+  final Function() onTap;
+
+  /// Whether the item is currently selected.
+  final bool bIsSelected;
+
+  /// Whether this can be selected as a target object.
+  bool get bCanBeSelected => node.data.type != _ColorGradingObjectEntryType.level;
+
+  /// The path of the icon to display.
+  String? get iconPath => _iconsByType[node.data.type];
+
+  @override
+  Widget build(BuildContext context) {
+    late final CardListTileExpansionState expansionState;
+    if (node.children.length > 0) {
+      expansionState = controller.isNodeExpanded(node.key)
+          ? CardListTileExpansionState.expanded
+          : CardListTileExpansionState.collapsed;
+    } else {
+      expansionState = CardListTileExpansionState.none;
+    }
+
+    Widget tileBuilder(context, bIsEnabled) => CardListTile(
+          bDeEmphasize: !bIsEnabled,
+          bIsSelected: bCanBeSelected ? bIsSelected : false,
+          title: node.data.name,
+          iconPath: iconPath,
+          onTap: bCanBeSelected ? _onTap : null,
+          onTapIcon: _onTapIcon,
+          indentation: node.indentation,
+          expansionState: expansionState,
+        );
+
+    final UnrealProperty? enableProperty = node.data.enableProperty;
+    if (enableProperty == null) {
+      return tileBuilder(context, true);
+    }
+
+    return _EnablePropertySwipeRevealer(
+      enableProperty: enableProperty,
+      builder: tileBuilder,
+    );
+  }
+
+  /// Called when the main body of the entry is tapped.
+  void _onTap() {
+    if (!bCanBeSelected) {
+      return;
+    }
+
+    onTap();
+  }
+
+  /// Called when the icon of the entry is tapped.
+  void _onTapIcon() {
+    // Re-enable this if/when we have icons for collapsed state
+    //controller.toggleNode(node.key);
+  }
+}
+
+/// Widget representing an entry in the list of color grading targets.
+class _ColorGradingTargetListEntry extends StatelessWidget {
+  const _ColorGradingTargetListEntry({
+    required this.data,
+    required this.onTap,
+    required this.bIsSelected,
+    Key? key,
+  }) : super(key: key);
+
+  final Function() onTap;
+  final _ColorGradingTargetListEntryData data;
+  final bool bIsSelected;
+
+  @override
+  Widget build(BuildContext context) {
+    return _EnablePropertySwipeRevealer(
+      enableProperty: data.enableProperty,
+      builder: (_, bIsEnabled) => CardListTile(
+        bIsSelected: bIsSelected,
+        title: data.name,
+        iconPath: 'assets/images/icons/viewport.svg',
+        bDeEmphasize: !bIsEnabled,
+        onTap: onTap,
+      ),
+    );
+  }
+}
+
+/// Custom [SwipeRevealer] that handles enabling/disabling the element based on a property value.
+class _EnablePropertySwipeRevealer extends StatelessWidget {
+  const _EnablePropertySwipeRevealer({
+    Key? key,
+    required this.enableProperty,
+    required this.builder,
+  }) : super(key: key);
+
+  /// The property controlling whether this is considered enabled.
+  final UnrealProperty enableProperty;
+
+  /// Function that builds the child based on whether the property is enabled.
+  final Widget Function(BuildContext context, bool bIsEnabled) builder;
+
+  @override
+  Widget build(BuildContext context) {
+    return UnrealPropertyBuilder<bool>(
+      property: enableProperty,
+      builder: (_, final bool? bIsEnabled, modify) => SwipeRevealer(
+        rightSwipeActionBuilder: (context, onFinished) => CardListTileSwipeAction(
+          iconPath: (bIsEnabled == true)
+              ? 'assets/images/icons/checkbox_opaque_checked.svg'
+              : 'assets/images/icons/checkbox_opaque_unchecked.svg',
+          color: UnrealColors.gray22,
+          iconSize: 18,
+          onPressed: () {
+            if (bIsEnabled != null) {
+              modify(
+                const SetOperation(),
+                description: AppLocalizations.of(context)!.transactionToggleColorGradingEnable,
+                deltaValue: !bIsEnabled,
+              );
+            }
+            onFinished();
+          },
+        ),
+        child: builder(context, bIsEnabled != false),
+      ),
+    );
+  }
+}
+
+/// Types of object that can be shown in the color grading tab's outliner panel.
+enum _ColorGradingObjectEntryType {
+  level,
+  nDisplayConfig,
+  icvfxCamera,
+  postProcessVolume,
+}
+
+/// Data about an entry in the color grading tab's outliner panel.
+class _ColorGradingObjectEntryData {
+  const _ColorGradingObjectEntryData({
+    required this.path,
+    required this.name,
+    required this.type,
+    required this.targets,
+    required this.targetListTitle,
+    this.parent,
+  });
+
+  /// The path of the actor or component this represents.
+  final String path;
+
+  /// The name to show in the outliner panel.
+  final String name;
+
+  /// The type of actor/component represented by this entry.
+  final _ColorGradingObjectEntryType type;
+
+  /// The title of the target list panel when this is selected.
+  final String targetListTitle;
+
+  /// The targets to display in the target list panel when this entry is selected.
+  final List<_ColorGradingTargetListEntryData> targets;
+
+  /// The object this is nested under in the Outliner panel, if any.
+  final _ColorGradingObjectEntryData? parent;
+
+  /// The property containing the list of targets, or null if there's no such property
+  UnrealProperty? get targetListProperty {
+    switch (type) {
+      case _ColorGradingObjectEntryType.nDisplayConfig:
+        return UnrealProperty(objectPath: path, propertyName: 'StageSettings.PerViewportColorGrading');
+
+      case _ColorGradingObjectEntryType.icvfxCamera:
+        return UnrealProperty(objectPath: path, propertyName: 'CameraSettings.PerNodeColorGrading');
+
+      default:
+        return null;
+    }
+  }
+
+  /// The property that enables/disables color grading for this object.
+  UnrealProperty? get enableProperty {
+    switch (type) {
+      case _ColorGradingObjectEntryType.nDisplayConfig:
+        return UnrealProperty(objectPath: path, propertyName: 'StageSettings.EnableColorGrading');
+
+      case _ColorGradingObjectEntryType.postProcessVolume:
+        return UnrealProperty(objectPath: path, propertyName: 'bEnabled');
+
+      case _ColorGradingObjectEntryType.icvfxCamera:
+        return UnrealProperty(objectPath: path, propertyName: 'CameraSettings.EnableInnerFrustumColorGrading');
+
+      default:
+        return null;
+    }
+  }
+}
+
+/// Data about an entry in the target list panel.
+class _ColorGradingTargetListEntryData {
+  const _ColorGradingTargetListEntryData({
+    required this.name,
+    required this.property,
+    required this.enableProperty,
+    required this.bIsListEntry,
+    this.nameProperty,
+  });
+
+  /// The name to display in the outliner.
+  final String name;
+
+  /// The color property to control.
+  final UnrealProperty property;
+
+  /// The property that enables/disables color grading for this property.
+  final UnrealProperty enableProperty;
+
+  /// The property that controls the name of this color grading target (if available).
+  final UnrealProperty? nameProperty;
+
+  /// If true, this is a color grading setting within a list (e.g. a specific viewport or camera node).
+  /// If false, this is entire cluster/camera color grading.
+  final bool bIsListEntry;
+}
diff --git a/Engine/Extras/VirtualProduction/EpicStageApp/lib/widgets/screens/main/tabs/details_tab.dart b/Engine/Extras/VirtualProduction/EpicStageApp/lib/widgets/screens/main/tabs/details_tab.dart
new file mode 100644
--- /dev/null
+++ b/Engine/Extras/VirtualProduction/EpicStageApp/lib/widgets/screens/main/tabs/details_tab.dart
@@ -0,0 +1,603 @@
+// Copyright Epic Games, Inc. All Rights Reserved.
+
+import 'dart:async';
+import 'dart:collection';
+
+import 'package:flutter/material.dart';
+import 'package:flutter_gen/gen_l10n/app_localizations.dart';
+import 'package:provider/provider.dart';
+import 'package:streaming_shared_preferences/streaming_shared_preferences.dart';
+
+import '../../../../models/actor_data/light_card_actor_data.dart';
+import '../../../../models/property_modify_operations.dart';
+import '../../../../models/settings/details_tab_settings.dart';
+import '../../../../models/settings/main_screen_settings.dart';
+import '../../../../models/settings/selected_actor_settings.dart';
+import '../../../../models/unreal_actor_manager.dart';
+import '../../../../models/unreal_types.dart';
+import '../../../../utilities/constants.dart';
+import '../../../../utilities/guarded_refresh_state.dart';
+import '../../../../utilities/preferences_bundle.dart';
+import '../../../elements/delta_slider.dart';
+import '../../../elements/dropdown_button.dart';
+import '../../../elements/dropdown_text.dart';
+import '../../../elements/empty_placeholder.dart';
+import '../../../elements/epic_icon_button.dart';
+import '../../../elements/layout/card.dart';
+import '../../../elements/reset_mode_button.dart';
+import '../../../elements/selector_bar.dart';
+import '../../../elements/stepper.dart';
+import '../../../elements/unreal_property_builder.dart';
+import '../sidebar/outliner_panel.dart';
+import 'base_color_tab.dart';
+
+/// Which type of actor is being shown in the details tab.
+enum _DetailsActorType {
+  lightCard,
+  colorCorrectWindow,
+  colorCorrectRegion,
+  multiple,
+}
+
+/// A tab that lets the user edit basic details of the selected actors.
+class DetailsTab extends StatefulWidget {
+  const DetailsTab({Key? key}) : super(key: key);
+
+  static const String iconPath = 'assets/images/icons/details.svg';
+
+  static String getTitle(BuildContext context) => AppLocalizations.of(context)!.tabTitleDetails;
+
+  @override
+  State<DetailsTab> createState() => _DetailsTabState();
+}
+
+class _DetailsTabState extends State<DetailsTab> with GuardedRefreshState {
+  late final UnrealActorManager _actorManager;
+  late final SelectedActorSettings _selectedActorSettings;
+  late final DetailsTabSettings _tabSettings;
+
+  /// Stream subscriptions to user settings.
+  final List<StreamSubscription> _settingSubscriptions = [];
+
+  /// Classes for which this tab can modify properties.
+  Set<String> get validClasses {
+    final Set<String> classes = {lightCardClassName, colorCorrectRegionClassName};
+    classes.addAll(colorCorrectWindowClassNames);
+    return classes;
+  }
+
+  @override
+  void initState() {
+    super.initState();
+
+    _actorManager = Provider.of<UnrealActorManager>(context, listen: false);
+    _actorManager.watchClassName(lightCardClassName, _onActorUpdate);
+
+    _selectedActorSettings = Provider.of<SelectedActorSettings>(context, listen: false);
+    _settingSubscriptions.addAll([
+      _selectedActorSettings.selectedActors.listen(refreshOnData),
+    ]);
+
+    _tabSettings = DetailsTabSettings(PreferencesBundle.of(context));
+  }
+
+  @override
+  void dispose() {
+    _actorManager.stopWatchingClassName(lightCardClassName, _onActorUpdate);
+
+    for (final StreamSubscription subscription in _settingSubscriptions) {
+      subscription.cancel();
+    }
+
+    super.dispose();
+  }
+
+  @override
+  Widget build(BuildContext context) {
+    final _DetailsActorType? actorType = _getActorType();
+
+    return Padding(
+      padding: EdgeInsets.all(cardMargin),
+      child: Row(
+        children: [
+          Expanded(
+            child: Card(
+              child: Column(children: [
+                CardLargeHeader(
+                  title: _getTitle(),
+                  subtitle: DetailsTab.getTitle(context),
+                  iconPath: _getIconPath(),
+                  trailing: const ResetModeButton(),
+                ),
+                Expanded(child: _createInnerContents(actorType)),
+              ]),
+            ),
+          ),
+          PreferenceBuilder(
+              preference: Provider.of<MainScreenSettings>(context, listen: false).bIsOutlinerPanelOpen,
+              builder: (context, final bool bIsOutlinerPanelOpen) {
+                return Row(children: [
+                  if (bIsOutlinerPanelOpen) const SizedBox(width: cardMargin),
+                  if (bIsOutlinerPanelOpen) OutlinerPanel(),
+                ]);
+              }),
+        ],
+      ),
+    );
+  }
+
+  /// Get the title to display at the top of the tab.
+  String? _getTitle() {
+    final Set<String> selectedActors = _selectedActorSettings.selectedActors.getValue();
+
+    if (selectedActors.length == 1) {
+      final UnrealObject? actor = _actorManager.getActorAtPath(selectedActors.first);
+
+      if (actor != null) {
+        return actor.name;
+      }
+    } else if (selectedActors.length > 1) {
+      return AppLocalizations.of(context)!.detailsTabMultipleActorsTitle;
+    }
+
+    return null;
+  }
+
+  /// Get the path for the icon to display at the top of the tab.
+  String? _getIconPath() {
+    String? iconPath;
+
+    if (_selectedActorSettings.selectedActors.getValue().length == 1) {
+      final UnrealObject? actor = _actorManager.getActorAtPath(_selectedActorSettings.selectedActors.getValue().first);
+
+      if (actor != null) {
+        iconPath = actor.getIconPath();
+      }
+    }
+
+    if (iconPath != null) {
+      return iconPath;
+    }
+
+    return 'assets/images/icons/details.svg';
+  }
+
+  /// Called when an actor we're editing has an update from the actor manager.
+  void _onActorUpdate(ActorUpdateDetails details) {
+    if (mounted && (details.renamedActors.isNotEmpty || details.addedActors.isNotEmpty)) {
+      // Force redraw in case we need to update the name/just got the name for an actor we were awaiting
+      setState(() {});
+    }
+  }
+
+  /// Determine the type of actors we're editing.
+  _DetailsActorType? _getActorType() {
+    _DetailsActorType? actorType;
+
+    final Set<String> selectedActors = _selectedActorSettings.selectedActors.getValue();
+    for (final String actorPath in selectedActors) {
+      final UnrealObject? actor = _actorManager.getActorAtPath(actorPath);
+      if (actor == null) {
+        continue;
+      }
+
+      _DetailsActorType? newActorType;
+
+      if (actor.isA(lightCardClassName)) {
+        newActorType = _DetailsActorType.lightCard;
+      }
+
+      if (actor.isA(colorCorrectRegionClassName)) {
+        newActorType = _DetailsActorType.colorCorrectRegion;
+      }
+
+      if (actor.isAny(colorCorrectWindowClassNames)) {
+        newActorType = _DetailsActorType.colorCorrectWindow;
+      }
+
+      if (actorType != null && newActorType != null && newActorType != actorType) {
+        return _DetailsActorType.multiple;
+      }
+
+      actorType = newActorType;
+    }
+
+    return actorType;
+  }
+
+  /// Create the inner contents of the tab (i.e. inside the panel UI).
+  Widget _createInnerContents(_DetailsActorType? actorType) {
+    if (actorType == null) {
+      final mainScreenSettings = Provider.of<MainScreenSettings>(context, listen: false);
+      return PreferenceBuilder(
+        preference: mainScreenSettings.bIsOutlinerPanelOpen,
+        builder: (context, final bool bIsOutlinerPanelOpen) => EmptyPlaceholder(
+          message: AppLocalizations.of(context)!.detailsTabEmptyMessage,
+          button: bIsOutlinerPanelOpen
+              ? null
+              : EpicWideButton(
+                  text: AppLocalizations.of(context)!.detailsTabShowOutlinerButtonLabel,
+                  iconPath: 'assets/images/icons/outliner.svg',
+                  onPressed: () => mainScreenSettings.bIsOutlinerPanelOpen.setValue(true),
+                ),
+        ),
+      );
+    }
+
+    if (actorType == _DetailsActorType.multiple) {
+      return EmptyPlaceholder(
+        message: AppLocalizations.of(context)!.detailsTabMixedActorsMessage,
+      );
+    }
+
+    // Buttons to toggle between property display types
+    Widget propertyDisplayTypeToggleButtons = PreferenceBuilder(
+      preference: _tabSettings.detailsPropertyDisplayType,
+      builder: (context, DetailsPropertyDisplayType detailsPropertyDisplayType) => SelectorBar(
+        value: detailsPropertyDisplayType,
+        onSelected: (DetailsPropertyDisplayType value) => _tabSettings.detailsPropertyDisplayType.setValue(value),
+        valueNames: {
+          DetailsPropertyDisplayType.orientation: AppLocalizations.of(context)!.colorTabOrientationPropertiesLabel,
+          DetailsPropertyDisplayType.appearance: AppLocalizations.of(context)!.colorTabAppearancePropertiesLabel,
+        },
+      ),
+    );
+
+    late final BaseColorTabMode mode;
+    late final List<UnrealProperty> colorProperties;
+    late final Widget otherPropertiesColumn;
+
+    switch (actorType) {
+      case _DetailsActorType.lightCard:
+        mode = BaseColorTabMode.color;
+        colorProperties = _getPropertiesOnValidActors('Color');
+        otherPropertiesColumn = _createLightCardPropertiesColumn();
+        break;
+
+      // We handle the OR case here because all CCWs are also CCRs, so they may be flagged as both
+      case _DetailsActorType.colorCorrectWindow:
+        mode = BaseColorTabMode.colorGrading;
+        colorProperties = _getPropertiesOnValidActors('ColorGradingSettings');
+        otherPropertiesColumn = _createColorCorrectWindowPropertiesColumn();
+        break;
+
+      case _DetailsActorType.colorCorrectRegion:
+        mode = BaseColorTabMode.colorGrading;
+        colorProperties = _getPropertiesOnValidActors('ColorGradingSettings');
+        otherPropertiesColumn = _createColorCorrectRegionPropertiesColumn();
+
+        // Can't control CCR orientation properties
+        propertyDisplayTypeToggleButtons =
+            FakeSelectorBar(AppLocalizations.of(context)!.colorTabAppearancePropertiesLabel);
+        break;
+
+      default:
+        throw 'Invalid actor type for top widget in details panel';
+    }
+
+    return BaseColorTab(
+      colorProperties: colorProperties,
+      mode: mode,
+      rightSideHeader: Center(child: propertyDisplayTypeToggleButtons),
+      rightSideContents: otherPropertiesColumn,
+    );
+  }
+
+  /// Create property widgets for a lightcard.
+  Widget _createLightCardPropertiesColumn() {
+    return PreferenceBuilder(
+      preference: _tabSettings.detailsPropertyDisplayType,
+      builder: (context, DetailsPropertyDisplayType detailsPropertyDisplayType) {
+        late List<Widget> widgets;
+
+        switch (detailsPropertyDisplayType) {
+          case DetailsPropertyDisplayType.orientation:
+            widgets = [
+              UnrealDropdownSelector(
+                overrideName: AppLocalizations.of(context)!.propertyLightCardMask,
+                unrealProperties: _getPropertiesOnValidActors('Mask'),
+              ),
+              for (final Widget slider in _createCommonStageActorOrientationPropertyWidgets()) slider,
+            ];
+            break;
+
+          case DetailsPropertyDisplayType.appearance:
+            widgets = [
+              UnrealDeltaSlider(
+                unrealProperties: _getPropertiesOnValidActors('Temperature'),
+              ),
+              UnrealDeltaSlider(
+                unrealProperties: _getPropertiesOnValidActors('Tint'),
+              ),
+              UnrealDeltaSlider(
+                unrealProperties: _getPropertiesOnValidActors('Gain'),
+                softMin: 0,
+                softMax: 10,
+              ),
+              UnrealDeltaSlider(
+                unrealProperties: _getPropertiesOnValidActors('Opacity'),
+              ),
+              UnrealDeltaSlider(
+                unrealProperties: _getPropertiesOnValidActors('Feathering'),
+              ),
+              UnrealDeltaSlider(
+                unrealProperties: _getPropertiesOnValidActors('Exposure'),
+                softMin: -10,
+                softMax: 10,
+              ),
+              UnrealStepper(
+                unrealProperties: _getPropertiesOnValidActors('Exposure'),
+                steps: StepperStepConfig.exposureSteps,
+              ),
+            ];
+            break;
+
+          default:
+            widgets = [];
+        }
+
+        return Column(children: widgets);
+      },
+    );
+  }
+
+  /// Create property widgets for a CCR.
+  Widget _createColorCorrectRegionPropertiesColumn() {
+    final temperatureTypeProperties = _getPropertiesOnValidActors('TemperatureType');
+
+    return Column(children: [
+      UnrealMultiPropertyBuilder<String>(
+        properties: temperatureTypeProperties,
+        fallbackValue: AppLocalizations.of(context)!.mismatchedValuesLabel,
+        builder: (_, String? sharedValue, __) => UnrealDeltaSlider(
+          overrideName: sharedValue,
+          unrealProperties: _getPropertiesOnValidActors('Temperature'),
+          buildLabel: (name) => UnrealDropdownText(unrealProperties: temperatureTypeProperties),
+        ),
+      ),
+      UnrealDeltaSlider(
+        unrealProperties: _getPropertiesOnValidActors('Tint'),
+      ),
+      UnrealDeltaSlider(
+        unrealProperties: _getPropertiesOnValidActors('Intensity'),
+      ),
+      UnrealDeltaSlider(
+        unrealProperties: _getPropertiesOnValidActors('Inner'),
+      ),
+      UnrealDeltaSlider(
+        unrealProperties: _getPropertiesOnValidActors('Outer'),
+      ),
+      UnrealDeltaSlider(
+        unrealProperties: _getPropertiesOnValidActors('Falloff'),
+      ),
+    ]);
+  }
+
+  /// Create property widgets for a CCW.
+  Widget _createColorCorrectWindowPropertiesColumn() {
+    return PreferenceBuilder(
+      preference: _tabSettings.detailsPropertyDisplayType,
+      builder: (context, DetailsPropertyDisplayType detailsPropertyDisplayType) {
+        late List<Widget> widgets;
+
+        switch (detailsPropertyDisplayType) {
+          case DetailsPropertyDisplayType.orientation:
+            widgets = [
+              UnrealDropdownSelector(
+                overrideName: AppLocalizations.of(context)!.propertyLightCardMask,
+                unrealProperties: _getPropertiesOnValidActors('WindowType'),
+              ),
+              for (final Widget slider in _createCommonStageActorOrientationPropertyWidgets()) slider,
+              UnrealDeltaSlider(
+                unrealProperties: _getPropertiesOnValidActors(
+                  'RadialOffset',
+                  modifierFunction: _modifyPositionalPropertyNameBasedOnClass,
+                ),
+              ),
+            ];
+            break;
+
+          default:
+            return _createColorCorrectRegionPropertiesColumn();
+        }
+
+        return Column(children: widgets);
+      },
+    );
+  }
+
+  /// Create property widgets for controlling the selected actors' orientations.
+  List<Widget> _createCommonStageActorOrientationPropertyWidgets() {
+    return _createScalePropertyWidgets() +
+        _createPositionPropertyWidgets() +
+        [
+          UnrealDeltaSlider(
+            key: Key('Spin'),
+            unrealProperties: _getPropertiesOnValidActors(
+              'Spin',
+              modifierFunction: _modifyPositionalPropertyNameBasedOnClass,
+            ),
+            minMaxBehaviour: PropertyMinMaxBehaviour.loop,
+          ),
+        ];
+  }
+
+  /// Create sliders for controlling the position of the selected actors.
+  List<Widget> _createPositionPropertyWidgets() {
+    // Classes for actor that can be positioned
+    final List<String> positionedActorClasses = [lightCardClassName];
+    positionedActorClasses.addAll(colorCorrectWindowClassNames);
+
+    // Determine which actors are UV/non-UV so we know which position sliders to show.
+    final List<String> uvActorPaths = [];
+    final List<String> nonUVActorPaths = [];
+
+    for (final String actorPath in _selectedActorSettings.selectedActors.getValue()) {
+      final UnrealObject? actor = _actorManager.getActorAtPath(actorPath);
+      if (!(actor?.isAny(positionedActorClasses) ?? false)) {
+        // Actor is not a class with position properties, so leave it out entirely
+        continue;
+      }
+
+      final LightCardActorData? lightCardActorData = actor!.getPerClassData<LightCardActorData>();
+      if (lightCardActorData?.bIsUV == true) {
+        uvActorPaths.add(actorPath);
+      } else {
+        nonUVActorPaths.add(actorPath);
+      }
+    }
+
+    final List<Widget> positionSliders = [];
+    if (nonUVActorPaths.isNotEmpty) {
+      positionSliders.addAll([
+        UnrealDeltaSlider(
+          key: Key('Latitude'),
+          unrealProperties: _getPropertiesOnActors(
+            'Latitude',
+            nonUVActorPaths,
+            modifierFunction: _modifyPositionalPropertyNameBasedOnClass,
+          ),
+          minMaxBehaviour: PropertyMinMaxBehaviour.ignore,
+        ),
+        UnrealDeltaSlider(
+          key: Key('Longitude'),
+          unrealProperties: _getPropertiesOnActors(
+            'Longitude',
+            nonUVActorPaths,
+            modifierFunction: _modifyPositionalPropertyNameBasedOnClass,
+          ),
+          minMaxBehaviour: PropertyMinMaxBehaviour.ignore,
+        ),
+      ]);
+    } else if (uvActorPaths.isNotEmpty) {
+      positionSliders.addAll([
+        UnrealDeltaSlider(
+          key: Key('UV X'),
+          unrealProperties: _getPropertiesOnActors('UVCoordinates.X', uvActorPaths),
+          overrideName: AppLocalizations.of(context)!.propertyLightCardUVX,
+          minMaxBehaviour: PropertyMinMaxBehaviour.clamp,
+        ),
+        UnrealDeltaSlider(
+          key: Key('UV Y'),
+          unrealProperties: _getPropertiesOnActors('UVCoordinates.Y', uvActorPaths),
+          overrideName: AppLocalizations.of(context)!.propertyLightCardUVY,
+          minMaxBehaviour: PropertyMinMaxBehaviour.clamp,
+        ),
+      ]);
+    }
+
+    return positionSliders;
+  }
+
+  /// Create sliders for controlling the scale of the selected actors.
+  List<Widget> _createScalePropertyWidgets() {
+    // Get actors that can be scaled
+    final List<String> scalableActorClasses = [lightCardClassName];
+    scalableActorClasses.addAll(colorCorrectWindowClassNames);
+
+    final List<String> scalableActorPaths = _getValidActorPathsOfClasses(scalableActorClasses);
+
+    final List<Widget> scaleSliders = [];
+    if (scalableActorPaths.isNotEmpty) {
+      scaleSliders.addAll([
+        UnrealDeltaSlider(
+          key: Key('Scale X'),
+          unrealProperties: _getPropertiesOnActors(
+            'Scale.X',
+            scalableActorPaths,
+            modifierFunction: _modifyPositionalPropertyNameBasedOnClass,
+          ),
+          overrideName: AppLocalizations.of(context)!.propertyLightCardScaleX,
+          hardMin: 0,
+          softMax: 10,
+        ),
+        UnrealDeltaSlider(
+          key: Key('Scale Y'),
+          unrealProperties: _getPropertiesOnActors(
+            'Scale.Y',
+            scalableActorPaths,
+            modifierFunction: _modifyPositionalPropertyNameBasedOnClass,
+          ),
+          overrideName: AppLocalizations.of(context)!.propertyLightCardScaleY,
+          hardMin: 0,
+          softMax: 10,
+        ),
+      ]);
+    }
+
+    return scaleSliders;
+  }
+
+  /// Given an [actorPath] and a [propertyName], return a modified positional property name accounting for the actor's
+  /// type.
+  String _modifyPositionalPropertyNameBasedOnClass(String actorPath, String propertyName) {
+    final UnrealObject? actor = _actorManager.getActorAtPath(actorPath);
+
+    if (actor != null && actor.isAny(colorCorrectWindowClassNames)) {
+      // CCWs store their positional properties in a sub-structure.
+      return 'PositionalParams.' + propertyName;
+    }
+
+    return propertyName;
+  }
+
+  /// Get a list of properties with the given [name] for all of the actors with paths in [actorPaths].
+  /// If [modifierFunction] is provided, it will be called for each [actorPath] and the [propertyName] of the property,
+  /// and its return value will be used in place of [name].
+  List<UnrealProperty> _getPropertiesOnActors(
+    String name,
+    List<String> actorPaths, {
+    String Function(String actorPath, String propertyName)? modifierFunction,
+    String? typeNameOverride,
+  }) {
+    return actorPaths
+        .map(
+          (actorPath) => UnrealProperty(
+            objectPath: actorPath,
+            propertyName: (modifierFunction != null) ? modifierFunction(actorPath, name) : name,
+            typeNameOverride: typeNameOverride,
+          ),
+        )
+        .toList();
+  }
+
+  /// Get a list of properties with the given [name] for all selected actors that belong to a valid class.
+  /// If [modifierFunction] is provided, it will be called for each [actorPath] and the [propertyName] of the property,
+  /// and its return value will be used in place of [name].
+  List<UnrealProperty> _getPropertiesOnValidActors(
+    String name, {
+    String Function(String actorPath, String propertyName)? modifierFunction,
+    String? typeNameOverride,
+  }) {
+    return _getPropertiesOnActors(
+      name,
+      _getValidActorPaths(),
+      modifierFunction: modifierFunction,
+      typeNameOverride: typeNameOverride,
+    );
+  }
+
+  /// Check if an actor is in the set of valid classes.
+  bool _isActorOfValidClass(String actorPath) {
+    final UnrealObject? actor = _actorManager.getActorAtPath(actorPath);
+    if (actor == null) {
+      return false;
+    }
+
+    final UnmodifiableSetView<String> actorClasses = actor.classNames;
+    return validClasses.any((className) => actorClasses.contains(className));
+  }
+
+  /// Get a list of selected actor paths that we want to edit (i.e. have valid classes).
+  List<String> _getValidActorPaths() {
+    return _selectedActorSettings.selectedActors
+        .getValue()
+        .where((actorPath) => _isActorOfValidClass(actorPath))
+        .toList();
+  }
+
+  /// Return all valid selected actors actors that are a member of at least one class in [classNames].
+  List<String> _getValidActorPathsOfClasses(List<String> classNames) => _getValidActorPaths()
+      .where((path) => _actorManager.getActorAtPath(path)?.isAny(classNames) ?? false)
+      .toList(growable: false);
+}
diff --git a/Engine/Extras/VirtualProduction/EpicStageApp/lib/widgets/screens/main/tabs/main_screen_tabs.dart b/Engine/Extras/VirtualProduction/EpicStageApp/lib/widgets/screens/main/tabs/main_screen_tabs.dart
new file mode 100644
--- /dev/null
+++ b/Engine/Extras/VirtualProduction/EpicStageApp/lib/widgets/screens/main/tabs/main_screen_tabs.dart
@@ -0,0 +1,94 @@
+// Copyright Epic Games, Inc. All Rights Reserved.
+
+import 'dart:collection';
+
+import 'package:flutter/material.dart';
+
+import 'color_grading_tab.dart';
+import 'details_tab.dart';
+import 'stage_map_tab.dart';
+import 'web_browser_tab.dart';
+
+/// The tab bar shown within the [MainScreenToolbar].
+class MainScreenTabs {
+  static final List<MainScreenTabConfig> _tabs = [
+    MainScreenTabConfig(
+      getTitle: StageMapTab.getTitle,
+      iconPath: StageMapTab.iconPath,
+      bShowMapPreview: false,
+      createTabContents: (context) => const StageMapTab(),
+    ),
+    MainScreenTabConfig(
+      getTitle: DetailsTab.getTitle,
+      iconPath: DetailsTab.iconPath,
+      createTabContents: (context) => const DetailsTab(),
+    ),
+    MainScreenTabConfig(
+      getTitle: ColorGradingTab.getTitle,
+      iconPath: ColorGradingTab.iconPath,
+      createTabContents: (context) => const ColorGradingTab(),
+    ),
+    MainScreenTabConfig(
+      getTitle: WebBrowserTab.getTitle,
+      iconPath: WebBrowserTab.iconPath,
+      bKeepAlive: true,
+      bEnableOutlinerToggle: false,
+      bShowMapPreview: false,
+      createTabContents: (context) => const WebBrowserTab(),
+    ),
+  ];
+
+  /// The list of tab configurations in display order.
+  static UnmodifiableListView<MainScreenTabConfig> get tabConfigs => UnmodifiableListView(_tabs);
+
+  /// Create the contents of one of the tabs by its index.
+  static Widget createTabContents(BuildContext context, int index) {
+    return _tabs[index].createTabContents(context);
+  }
+
+  /// Check whether the given tab's contents should be kept alive when not visible.
+  static bool shouldKeepTabAlive(int index) {
+    return _tabs[index].bKeepAlive;
+  }
+
+  /// Check whether the outliner can be toggled in the given tab.
+  static bool shouldEnableOutlinerToggle(int index) {
+    return _tabs[index].bEnableOutlinerToggle;
+  }
+
+  /// Check whether the given tab should show the map preview.
+  static bool shouldShowMapPreview(int index) {
+    return _tabs[index].bShowMapPreview;
+  }
+}
+
+/// Configuration for an entry in the main screen's tab bar.
+class MainScreenTabConfig {
+  const MainScreenTabConfig({
+    required this.getTitle,
+    required this.iconPath,
+    required this.createTabContents,
+    this.bKeepAlive = false,
+    this.bEnableOutlinerToggle = true,
+    this.bShowMapPreview = true,
+  });
+
+  /// Function that returns the title of the tab as shown to the user.
+  final String Function(BuildContext) getTitle;
+
+  /// The path of the icon representing this tab.
+  final String iconPath;
+
+  /// A function to create the widget for the tab's contents.
+  final Widget Function(BuildContext context) createTabContents;
+
+  /// If true, keep this tab's contents alive even when it's not visible.
+  /// If false, dispose of the tab's contents when it's not visible.
+  final bool bKeepAlive;
+
+  /// Whether the outliner can be toggled in this tab.
+  final bool bEnableOutlinerToggle;
+
+  /// Whether the floating map preview should be visible when this tab is open.
+  final bool bShowMapPreview;
+}
diff --git a/Engine/Extras/VirtualProduction/EpicStageApp/lib/widgets/screens/main/tabs/stage_map_tab.dart b/Engine/Extras/VirtualProduction/EpicStageApp/lib/widgets/screens/main/tabs/stage_map_tab.dart
new file mode 100644
--- /dev/null
+++ b/Engine/Extras/VirtualProduction/EpicStageApp/lib/widgets/screens/main/tabs/stage_map_tab.dart
@@ -0,0 +1,69 @@
+// Copyright Epic Games, Inc. All Rights Reserved.
+
+import 'package:flutter/material.dart';
+import 'package:flutter_gen/gen_l10n/app_localizations.dart';
+import 'package:provider/provider.dart';
+import 'package:streaming_shared_preferences/streaming_shared_preferences.dart';
+
+import '../../../../models/settings/main_screen_settings.dart';
+import '../../../../models/unreal_types.dart';
+import '../../../../utilities/constants.dart';
+import '../../../../utilities/unreal_colors.dart';
+import '../../../elements/lightcard_map.dart';
+import '../sidebar/outliner_panel.dart';
+
+class StageMapTab extends StatefulWidget {
+  const StageMapTab({Key? key}) : super(key: key);
+  static const String iconPath = 'assets/images/icons/world.svg';
+
+  static String getTitle(BuildContext context) => AppLocalizations.of(context)!.tabTitleStageMap;
+
+  @override
+  State<StageMapTab> createState() => _StageMapTabState();
+}
+
+class _StageMapTabState extends State<StageMapTab> {
+  final _mapKey = GlobalKey();
+
+  /// The state of the map this is displaying.
+  StageMapState get _mapState => _mapKey.currentState as StageMapState;
+
+  @override
+  Widget build(BuildContext context) {
+    return PreferenceBuilder(
+      preference: Provider.of<MainScreenSettings>(context).bIsOutlinerPanelOpen,
+      builder: (context, final bool bIsOutlinerPanelOpen) => Container(
+        color: UnrealColors.black,
+        child: Row(children: [
+          Expanded(
+            key: const Key('StageMap'),
+            child: StageMap(key: _mapKey),
+          ), // Outliner panel
+
+          if (bIsOutlinerPanelOpen)
+            Container(
+              key: const Key('OutlinerPanel'),
+              padding: EdgeInsets.all(cardMargin),
+              decoration: BoxDecoration(
+                borderRadius: const BorderRadius.only(
+                  topLeft: Radius.circular(outerCornerRadius),
+                  bottomLeft: Radius.circular(outerCornerRadius),
+                ),
+                color: Theme.of(context).colorScheme.background,
+              ),
+              child: OutlinerPanel(
+                focusActor: _focusActor,
+                canFocusActor: _getCanFocusActor,
+              ),
+            ),
+        ]),
+      ),
+    );
+  }
+
+  /// Focus the map on a specific [actor].
+  void _focusActor(UnrealObject actor) => _mapState.focusActor(actor);
+
+  /// Check whether the map can focus on an [actor].
+  bool _getCanFocusActor(UnrealObject actor) => _mapState.getCanFocusActor(actor);
+}
diff --git a/Engine/Extras/VirtualProduction/EpicStageApp/lib/widgets/screens/main/tabs/web_browser_tab.dart b/Engine/Extras/VirtualProduction/EpicStageApp/lib/widgets/screens/main/tabs/web_browser_tab.dart
new file mode 100644
--- /dev/null
+++ b/Engine/Extras/VirtualProduction/EpicStageApp/lib/widgets/screens/main/tabs/web_browser_tab.dart
@@ -0,0 +1,319 @@
+// Copyright Epic Games, Inc. All Rights Reserved.
+
+import 'package:flutter/material.dart';
+import 'package:flutter_gen/gen_l10n/app_localizations.dart';
+import 'package:logging/logging.dart';
+import 'package:provider/provider.dart';
+import 'package:webview_flutter/webview_flutter.dart';
+
+import '../../../../models/engine_connection.dart';
+import '../../../../utilities/constants.dart';
+import '../../../../utilities/net_utilities.dart';
+import '../../../elements/empty_placeholder.dart';
+import '../../../elements/epic_icon_button.dart';
+import '../../../elements/spinner_overlay.dart';
+
+final _log = Logger('WebBrowser');
+
+class WebBrowserTab extends StatefulWidget {
+  const WebBrowserTab({Key? key}) : super(key: key);
+
+  static const String iconPath = 'assets/images/icons/web_browser.svg';
+
+  static String getTitle(BuildContext context) => AppLocalizations.of(context)!.tabTitleWebBrowser;
+
+  /// Browser tabs states that are currently active in the hierarchy.
+  static final Set<_WebBrowserTabState> _activeBrowserTabs = {};
+
+  @override
+  State<WebBrowserTab> createState() => _WebBrowserTabState();
+
+  /// Refresh all active web browser tabs.
+  static void refreshAll() {
+    for (_WebBrowserTabState tabState in _activeBrowserTabs) {
+      tabState.refresh();
+    }
+  }
+}
+
+class _WebBrowserTabState extends State<WebBrowserTab> {
+  late final EngineConnectionManager _connectionManager;
+
+  /// Controls the state of the web view.
+  WebViewController? _webViewController;
+
+  /// The initial URL to attempt to connect to.
+  String? _initialUrl;
+
+  /// The current URL to display in the address bar.
+  String _currentUrl = '';
+
+  /// The error message to display. If set, this will be shown instead of the web view.
+  String? _errorMessage;
+
+  /// Whether the page has finished its initial load.
+  bool _bIsPageLoaded = false;
+
+  @override
+  void initState() {
+    super.initState();
+
+    _connectionManager = Provider.of<EngineConnectionManager>(context, listen: false);
+    WidgetsBinding.instance.addPostFrameCallback((_) => _attemptToConnect());
+
+    WebBrowserTab._activeBrowserTabs.add(this);
+  }
+
+  @override
+  void activate() {
+    super.activate();
+
+    WebBrowserTab._activeBrowserTabs.add(this);
+  }
+
+  @override
+  void deactivate() {
+    WebBrowserTab._activeBrowserTabs.remove(this);
+
+    super.deactivate();
+  }
+
+  @override
+  void dispose() {
+    WebBrowserTab._activeBrowserTabs.remove(this);
+
+    super.dispose();
+  }
+
+  @override
+  Widget build(BuildContext context) {
+    late final Widget mainContent;
+
+    if (_errorMessage == null) {
+      // Show web view
+      mainContent = Column(
+        crossAxisAlignment: CrossAxisAlignment.stretch,
+        children: [
+          _AddressBar(address: _currentUrl),
+          const SizedBox(height: sectionMargin),
+          Expanded(
+            child: Stack(
+              children: [
+                if (!_bIsPageLoaded) const SpinnerOverlay(),
+                if (_initialUrl != null)
+                  Center(
+                    child: WebView(
+                      backgroundColor: Theme.of(context).colorScheme.surface,
+                      initialUrl: _initialUrl!,
+                      javascriptMode: JavascriptMode.unrestricted,
+                      allowsInlineMediaPlayback: true,
+                      onPageFinished: _onPageFinished,
+                      onWebViewCreated: (controller) => _webViewController = controller,
+                      onPageStarted: (String url) => setState(() {
+                        _currentUrl = url;
+                      }),
+                    ),
+                  ),
+              ],
+            ),
+          ),
+        ],
+      );
+    } else {
+      // Show error message
+      mainContent = Center(
+        child: EmptyPlaceholder(
+          message: _errorMessage!,
+          button: EpicWideButton(
+            text: AppLocalizations.of(context)!.webBrowserReconnectButtonLabel,
+            iconPath: 'assets/images/icons/refresh.svg',
+            onPressed: _attemptToConnect,
+          ),
+        ),
+      );
+    }
+
+    return Padding(
+      padding: const EdgeInsets.all(cardMargin),
+      child: Card(child: mainContent),
+    );
+  }
+
+  /// Refresh the currently loaded page, or attempt to reconnect if not currently showing a page.
+  void refresh() {
+    if (_webViewController == null) {
+      _attemptToConnect();
+      return;
+    }
+
+    _webViewController!.reload();
+  }
+
+  /// Try the full connection process from scratch, including retrieving the URI from the engine.
+  void _attemptToConnect() {
+    _initialUrl = null;
+    _errorMessage = null;
+    _bIsPageLoaded = false;
+
+    try {
+      WebView.platform;
+    } catch (e) {
+      _errorMessage = AppLocalizations.of(context)!.webBrowserUnsupportedMessage;
+      setState(() {});
+    }
+
+    if (_errorMessage == null) {
+      _findBrowserURI().then((bool bSuccess) {
+        if (!bSuccess) {
+          setState(() {
+            _errorMessage = AppLocalizations.of(context)!.webBrowserConnectionDataFailedMessage;
+          });
+        }
+      });
+    }
+  }
+
+  /// Request the port information for the web interface from the engine and construct the URL in [_initialUrl] if
+  /// found.
+  Future<bool> _findBrowserURI() async {
+    final ConnectionData? connectionData = await _connectionManager.getLastConnectionData();
+
+    if (connectionData == null) {
+      _log.warning('');
+      return false;
+    }
+
+    if (_connectionManager.connectionState != EngineConnectionState.connected) {
+      return false;
+    }
+
+    final UnrealHttpRequest request;
+
+    final bool bUseNewApi = _connectionManager.apiVersion!.bHasWebInterfacePortFunction;
+    if (bUseNewApi) {
+      request = UnrealHttpRequest(url: '/remote/object/call', verb: 'PUT', body: {
+        'objectPath': '/Script/EpicStageApp.Default__StageAppFunctionLibrary',
+        'functionName': 'GetRemoteControlWebInterfacePort',
+        'generateTransaction': false,
+      });
+    } else {
+      request = UnrealHttpRequest(url: '/remote/object/property', verb: 'PUT', body: {
+        'objectPath': '/Script/RemoteControlCommon.Default__RemoteControlSettings',
+        'propertyName': 'RemoteControlWebInterfacePort',
+        'access': 'READ_ACCESS',
+      });
+    }
+
+    final bool bSuccess = await _connectionManager.sendHttpRequest(request).then((UnrealHttpResponse response) {
+      if (response.code != 200) {
+        setState(() {
+          _log.warning('Request for web interface port failed (error ${response.code}).');
+
+          final String? responseError = response.body?['errorMessage'];
+          if (responseError != null) {
+            _errorMessage = _errorMessage! + '\n\n$responseError';
+          }
+        });
+        return false;
+      }
+
+      final int? port = response.body?[bUseNewApi ? 'ReturnValue' : 'RemoteControlWebInterfacePort'];
+      if (port == null) {
+        setState(() {
+          _log.warning('Response from engine did not contain a web interface port.');
+        });
+        return false;
+      }
+
+      setState(() {
+        _initialUrl = 'http://${connectionData.websocketAddress.address}:$port';
+      });
+
+      return true;
+    }).onError((error, stackTrace) {
+      setState(() {
+        _log.warning('Failed to send request for web interface port.');
+      });
+      return false;
+    });
+
+    return bSuccess;
+  }
+
+  /// Called when the web view finishes loading a page.
+  void _onPageFinished(String url) {
+    if (!mounted) {
+      return;
+    }
+
+    setState(() {
+      _bIsPageLoaded = true;
+    });
+  }
+}
+
+/// Address bar displayed at the top of the web browser tab.
+class _AddressBar extends StatelessWidget {
+  const _AddressBar({Key? key, required this.address}) : super(key: key);
+
+  final String address;
+
+  @override
+  Widget build(BuildContext context) {
+    return Container(
+      height: 44,
+      color: Theme.of(context).colorScheme.surfaceTint,
+      child: Padding(
+        padding: const EdgeInsets.symmetric(horizontal: cardMargin),
+        child: Row(
+          crossAxisAlignment: CrossAxisAlignment.center,
+          mainAxisAlignment: MainAxisAlignment.center,
+          children: [
+            Flexible(
+              child: Padding(
+                padding: const EdgeInsets.symmetric(vertical: cardMargin),
+                child: Container(
+                  constraints: BoxConstraints(minWidth: 300, maxWidth: 500),
+                  decoration: BoxDecoration(
+                    color: Theme.of(context).colorScheme.background,
+                    borderRadius: BorderRadius.circular(outerCornerRadius),
+                  ),
+                  child: Center(
+                    child: Padding(
+                      padding: EdgeInsets.symmetric(horizontal: 8),
+                      child: Text(
+                        address,
+                        style: Theme.of(context).textTheme.bodyMedium!.copyWith(
+                              height: 1.2,
+                              overflow: TextOverflow.fade,
+                            ),
+                        maxLines: 1,
+                        softWrap: false,
+                      ),
+                    ),
+                  ),
+                ),
+              ),
+            ),
+            const _RefreshWebBrowserButton(),
+          ],
+        ),
+      ),
+    );
+  }
+}
+
+/// Button that refreshes any open browser tabs.
+class _RefreshWebBrowserButton extends StatelessWidget {
+  const _RefreshWebBrowserButton({Key? key}) : super(key: key);
+
+  @override
+  Widget build(BuildContext context) {
+    return EpicIconButton(
+      iconPath: 'assets/images/icons/refresh.svg',
+      tooltipMessage: AppLocalizations.of(context)!.webBrowserRefreshButtonTooltip,
+      onPressed: WebBrowserTab.refreshAll,
+      buttonSize: const Size(40, 40),
+    );
+  }
+}
diff --git a/Engine/Extras/VirtualProduction/EpicStageApp/lib/widgets/screens/main/toolbar/main_screen_tab_bar.dart b/Engine/Extras/VirtualProduction/EpicStageApp/lib/widgets/screens/main/toolbar/main_screen_tab_bar.dart
new file mode 100644
--- /dev/null
+++ b/Engine/Extras/VirtualProduction/EpicStageApp/lib/widgets/screens/main/toolbar/main_screen_tab_bar.dart
@@ -0,0 +1,63 @@
+// Copyright Epic Games, Inc. All Rights Reserved.
+
+import 'package:flutter/material.dart';
+
+import '../../../../utilities/unreal_colors.dart';
+import '../../../elements/asset_icon.dart';
+import '../tabs/main_screen_tabs.dart';
+
+/// The tab bar shown within the [MainScreenToolbar].
+class MainScreenTabBar extends StatefulWidget {
+  const MainScreenTabBar({
+    Key? key,
+    this.tabController,
+  }) : super(key: key);
+
+  final TabController? tabController;
+
+  @override
+  State<MainScreenTabBar> createState() => _MainScreenTabBarState();
+}
+
+class _MainScreenTabBarState extends State<MainScreenTabBar> {
+  @override
+  Widget build(BuildContext context) {
+    return TabBar(
+      controller: widget.tabController,
+      isScrollable: true, // This shrinks the tab buttons to fit their contents
+      unselectedLabelColor: UnrealColors.white,
+      padding: EdgeInsets.zero,
+      indicatorPadding: EdgeInsets.zero,
+      labelPadding: const EdgeInsets.symmetric(horizontal: 4),
+      indicator: BoxDecoration(
+        color: Theme.of(context).colorScheme.primary,
+      ),
+      indicatorSize: TabBarIndicatorSize.label,
+      tabs: MainScreenTabs.tabConfigs.map((tabConfig) => _MainScreenTabBarTab(tabConfig: tabConfig)).toList(),
+    );
+  }
+}
+
+class _MainScreenTabBarTab extends StatelessWidget {
+  const _MainScreenTabBarTab({Key? key, required this.tabConfig}) : super(key: key);
+
+  final MainScreenTabConfig tabConfig;
+
+  @override
+  Widget build(BuildContext context) {
+    return Tooltip(
+      message: tabConfig.getTitle(context),
+      child: SizedBox(
+        width: 48,
+        height: 48,
+        child: Center(
+          child: AssetIcon(
+            path: tabConfig.iconPath,
+            size: 24,
+            color: DefaultTextStyle.of(context).style.color,
+          ),
+        ),
+      ),
+    );
+  }
+}
