Branch snapshot for CL 27138787

diff --git a/Engine/Extras/VirtualProduction/EpicStageApp/lib/widgets/screens/main/toolbar/main_screen_toolbar.dart b/Engine/Extras/VirtualProduction/EpicStageApp/lib/widgets/screens/main/toolbar/main_screen_toolbar.dart
new file mode 100644
--- /dev/null
+++ b/Engine/Extras/VirtualProduction/EpicStageApp/lib/widgets/screens/main/toolbar/main_screen_toolbar.dart
@@ -0,0 +1,82 @@
+// Copyright Epic Games, Inc. All Rights Reserved.
+
+import 'package:flutter/material.dart';
+import 'package:flutter_gen/gen_l10n/app_localizations.dart';
+
+import '../../../elements/control_lock.dart';
+import '../../../elements/epic_icon_button.dart';
+import '../../../elements/place_actor_menu.dart';
+import '../../../elements/sensitivity_slider.dart';
+import '../sidebar/outliner_panel.dart';
+import 'main_screen_tab_bar.dart';
+import 'settings/settings_dialog.dart';
+
+const double _buttonSpacing = 8;
+
+/// The toolbar shown at the top of the [StageAppMainScreen].
+class MainScreenToolbar extends StatelessWidget implements PreferredSizeWidget {
+  const MainScreenToolbar({
+    Key? key,
+    required this.bEnableOutlinerToggle,
+    required this.tabController,
+    this.placeActorButtonKey,
+  }) : super(key: key);
+
+  /// Whether the outliner can be toggled in the current tab.
+  final bool bEnableOutlinerToggle;
+
+  /// Controller for which tab is selected.
+  final TabController tabController;
+
+  /// A key that can be used to refer to the place actors button within this toolbar.
+  final Key? placeActorButtonKey;
+
+  @override
+  Size get preferredSize => const Size.fromHeight(48);
+
+  @override
+  Widget build(BuildContext context) {
+    return Container(
+      padding: const EdgeInsets.only(right: 8),
+      decoration: BoxDecoration(
+        color: Theme.of(context).colorScheme.surface,
+      ),
+      height: preferredSize.height,
+      child: Row(
+        children: [
+          PlaceActorButton(key: placeActorButtonKey),
+          Container(
+            color: Theme.of(context).colorScheme.background,
+            width: 4,
+          ),
+          const SizedBox(width: 4),
+          MainScreenTabBar(tabController: tabController),
+          const Spacer(),
+          Wrap(
+            spacing: _buttonSpacing,
+            children: [
+              OutlinerToggleButton(bEnabled: bEnableOutlinerToggle),
+              const SensitivitySliderButton(),
+              const ControlLock(),
+              const SettingsButton(),
+            ],
+          ),
+        ],
+      ),
+    );
+  }
+}
+
+/// Button that opens the settings menu.
+class SettingsButton extends StatelessWidget {
+  const SettingsButton({Key? key}) : super(key: key);
+
+  @override
+  Widget build(BuildContext context) {
+    return EpicIconButton(
+      tooltipMessage: AppLocalizations.of(context)!.menuSettings,
+      iconPath: 'assets/images/icons/settings.svg',
+      onPressed: SettingsDialog.show,
+    );
+  }
+}
diff --git a/Engine/Extras/VirtualProduction/EpicStageApp/lib/widgets/screens/main/toolbar/settings/pages/settings_dialog_main.dart b/Engine/Extras/VirtualProduction/EpicStageApp/lib/widgets/screens/main/toolbar/settings/pages/settings_dialog_main.dart
new file mode 100644
--- /dev/null
+++ b/Engine/Extras/VirtualProduction/EpicStageApp/lib/widgets/screens/main/toolbar/settings/pages/settings_dialog_main.dart
@@ -0,0 +1,132 @@
+// Copyright Epic Games, Inc. All Rights Reserved.
+
+import 'package:flutter/material.dart';
+import 'package:flutter_gen/gen_l10n/app_localizations.dart';
+import 'package:provider/provider.dart';
+import 'package:url_launcher/url_launcher.dart';
+
+import '../../../../../../models/engine_connection.dart';
+import '../../../../../../models/navigator_keys.dart';
+import '../../../../../../models/settings/selected_actor_settings.dart';
+import '../../../../../../models/unreal_actor_manager.dart';
+import '../../../../../../utilities/unreal_colors.dart';
+import '../../../../../elements/asset_icon.dart';
+import '../../../../eula/eula_screen.dart';
+import '../settings_generic.dart';
+import 'settings_log_list.dart';
+import 'settings_root_actor_picker.dart';
+
+/// Main screen of the [SettingsDialog].
+class SettingsDialogMain extends StatelessWidget {
+  static const String route = '/';
+
+  const SettingsDialogMain({Key? key}) : super(key: key);
+
+  final String _documentationUri =
+      'https://dev.epicgames.com/community/learning/courses/OKO/unreal-engine-unreal-stage-ios-app/';
+
+  @override
+  Widget build(BuildContext context) {
+    final actorManager = Provider.of<UnrealActorManager>(context, listen: false);
+    final connection = Provider.of<EngineConnectionManager>(context, listen: false);
+
+    final bool bIsConnected = connection.connectionState == EngineConnectionState.connected;
+    final localizations = AppLocalizations.of(context)!;
+
+    return SettingsPageScaffold(
+      title: AppLocalizations.of(context)!.settingsDialogTitle,
+      body: Column(
+        mainAxisSize: MainAxisSize.min,
+        children: [
+          if (bIsConnected)
+            StreamBuilder(
+              stream: Provider.of<SelectedActorSettings>(context, listen: false).displayClusterRootPath,
+              builder: (_, final AsyncSnapshot<String> rootPath) => SettingsMenuItem(
+                title: localizations.settingsDialogRootActorSettingLabel,
+                iconPath: 'assets/images/icons/ndisplay.svg',
+                trailing: Text(actorManager.getActorAtPath(rootPath.data ?? '')?.name ?? ''),
+                onTap: () => Navigator.of(context).pushNamed(SettingsDialogRootActorPicker.route),
+              ),
+            ),
+          if (bIsConnected) const SettingsMenuDivider(),
+          SettingsMenuItem(
+            title: localizations.settingsDialogApplicationLogLabel,
+            iconPath: 'assets/images/icons/log.svg',
+            onTap: () => Navigator.of(context).pushNamed(SettingsLogList.route),
+          ),
+          const SettingsMenuDivider(),
+          SettingsMenuItem(
+            title: localizations.settingsDialogApplicationHelpLabel,
+            iconPath: 'assets/images/icons/help.svg',
+            trailingIconPath: null,
+            onTap: () => launchUrl(Uri.parse(_documentationUri)),
+          ),
+          SettingsMenuItem(
+            title: localizations.settingsDialogApplicationAboutLabel,
+            iconPath: 'assets/images/icons/info.svg',
+            onTap: () => rootNavigatorKey.currentState?.pushNamed(EulaScreen.route, arguments: {'onPressed': () {}}),
+          ),
+          if (bIsConnected) const SettingsMenuDivider(),
+          if (bIsConnected)
+            Container(
+              padding: EdgeInsets.all(8),
+              child: Row(
+                children: const [_DisconnectButton()],
+              ),
+            ),
+        ],
+      ),
+    );
+  }
+}
+
+/// A button that disconnects from the engine.
+class _DisconnectButton extends StatelessWidget {
+  const _DisconnectButton({Key? key}) : super(key: key);
+
+  @override
+  Widget build(BuildContext context) {
+    return Semantics(
+      button: true,
+      child: MouseRegion(
+        cursor: MaterialStateMouseCursor.clickable,
+        child: GestureDetector(
+          onTap: () => _disconnect(context),
+          behavior: HitTestBehavior.opaque,
+          child: Container(
+            height: 36,
+            padding: EdgeInsets.symmetric(horizontal: 24),
+            decoration: BoxDecoration(
+              borderRadius: BorderRadius.circular(20),
+              border: Border.all(
+                width: 2,
+                color: UnrealColors.warning,
+              ),
+            ),
+            child: Center(
+              child: Row(children: [
+                AssetIcon(
+                  path: 'assets/images/icons/exit.svg',
+                  size: 24,
+                  color: UnrealColors.warning,
+                ),
+                const SizedBox(width: 8),
+                Text(
+                  AppLocalizations.of(context)!.settingsDialogDisconnectButtonLabel,
+                  style: Theme.of(context).textTheme.headlineSmall!.copyWith(
+                        color: UnrealColors.warning,
+                      ),
+                ),
+              ]),
+            ),
+          ),
+        ),
+      ),
+    );
+  }
+
+  /// Disconnect from the engine
+  void _disconnect(BuildContext context) {
+    Provider.of<EngineConnectionManager>(context, listen: false).disconnect();
+  }
+}
diff --git a/Engine/Extras/VirtualProduction/EpicStageApp/lib/widgets/screens/main/toolbar/settings/pages/settings_log_list.dart b/Engine/Extras/VirtualProduction/EpicStageApp/lib/widgets/screens/main/toolbar/settings/pages/settings_log_list.dart
new file mode 100644
--- /dev/null
+++ b/Engine/Extras/VirtualProduction/EpicStageApp/lib/widgets/screens/main/toolbar/settings/pages/settings_log_list.dart
@@ -0,0 +1,90 @@
+// Copyright Epic Games, Inc. All Rights Reserved.
+
+import 'dart:io';
+
+import '../../../../../elements/epic_list_view.dart';
+import 'package:flutter/material.dart';
+import 'package:flutter_gen/gen_l10n/app_localizations.dart';
+
+import '../../../../../../utilities/guarded_refresh_state.dart';
+import '../../../../../../utilities/logging.dart';
+import '../settings_generic.dart';
+import 'settings_log_view.dart';
+
+/// Page showing the list of log files.
+class SettingsLogList extends StatefulWidget {
+  const SettingsLogList({Key? key}) : super(key: key);
+
+  static const String route = '/logs';
+
+  @override
+  State<SettingsLogList> createState() => _SettingsLogListState();
+}
+
+class _SettingsLogListState extends State<SettingsLogList> with GuardedRefreshState {
+  /// Scroll controller for the scroll view of the logs.
+  final _scrollController = ScrollController();
+
+  /// List of files to display.
+  List<File> _logFiles = [];
+
+  @override
+  void initState() {
+    super.initState();
+
+    _updateLogFileList();
+  }
+
+  @override
+  void dispose() {
+    _scrollController.dispose();
+    super.dispose();
+  }
+
+  /// Update the list of log files.
+  void _updateLogFileList() async {
+    final Directory logDirectory = await Logging.instance.getLogDirectory();
+    final List<File> files =
+        await logDirectory.list(recursive: false).map((FileSystemEntity entity) => File(entity.path)).toList();
+
+    // Sort so the most recent files appear first. We assume that files are named identically aside from their dates,
+    // which are assumed to be in lexicographical order (YYYY_MM_DD)
+    files.sort((File fileA, File fileB) => -fileA.path.compareTo(fileB.path));
+
+    setState(() {
+      _logFiles = files;
+    });
+  }
+
+  @override
+  Widget build(BuildContext context) {
+    return SettingsPageScaffold(
+      title: AppLocalizations.of(context)!.settingsDialogApplicationLogLabel,
+      body: SizedBox(
+        height: 240,
+        child: MediaQuery.removePadding(
+          removeTop: true,
+          removeBottom: true,
+          context: context,
+          child: EpicListView(
+            itemCount: _logFiles.length,
+            itemBuilder: (context, index) {
+              final File file = _logFiles[index];
+
+              return SettingsMenuItem(
+                title: Logging.getNameForLog(file),
+                iconPath: 'assets/images/icons/log.svg',
+                onTap: () => Navigator.of(context).pushNamed(
+                  SettingsLogView.route,
+                  arguments: SettingsLogViewArguments(
+                    file: file,
+                  ),
+                ),
+              );
+            },
+          ),
+        ),
+      ),
+    );
+  }
+}
diff --git a/Engine/Extras/VirtualProduction/EpicStageApp/lib/widgets/screens/main/toolbar/settings/pages/settings_log_view.dart b/Engine/Extras/VirtualProduction/EpicStageApp/lib/widgets/screens/main/toolbar/settings/pages/settings_log_view.dart
new file mode 100644
--- /dev/null
+++ b/Engine/Extras/VirtualProduction/EpicStageApp/lib/widgets/screens/main/toolbar/settings/pages/settings_log_view.dart
@@ -0,0 +1,391 @@
+// Copyright Epic Games, Inc. All Rights Reserved.
+
+import 'dart:io';
+
+import '../../../../../elements/epic_list_view.dart';
+import 'package:flutter/material.dart';
+import 'package:share_plus/share_plus.dart';
+
+import '../../../../../../utilities/external_notifier.dart';
+import '../../../../../../utilities/guarded_refresh_state.dart';
+import '../../../../../../utilities/logging.dart';
+import '../../../../../elements/asset_icon.dart';
+import '../../../../../elements/epic_icon_button.dart';
+import '../settings_generic.dart';
+
+/// Arguments passed to the route when a [SettingsLogView] is pushed to the navigation stack.
+class SettingsLogViewArguments {
+  const SettingsLogViewArguments({required this.file});
+
+  /// The file to open.
+  final File file;
+}
+
+/// Page showing the full view of a log.
+class SettingsLogView extends StatefulWidget {
+  const SettingsLogView({Key? key}) : super(key: key);
+
+  static const String route = '/log_view';
+
+  @override
+  State<SettingsLogView> createState() => _SettingsLogViewState();
+}
+
+class _SettingsLogViewState extends State<SettingsLogView> with GuardedRefreshState {
+  final shareButtonKey = GlobalKey();
+
+  @override
+  Widget build(BuildContext context) {
+    final arguments = ModalRoute.of(context)?.settings.arguments as SettingsLogViewArguments;
+
+    return SettingsPageScaffold(
+      title: Logging.getNameForLog(arguments.file),
+      titleBarTrailing: EpicIconButton(
+        key: shareButtonKey,
+        iconPath: 'assets/images/icons/share.svg',
+        color: Theme.of(context).colorScheme.primary,
+        onPressed: () => _shareLog(arguments.file),
+      ),
+      body: SizedBox(
+        height: MediaQuery.of(context).size.height - 120,
+        child: FutureBuilder(
+          future: arguments.file.readAsString(),
+          builder: (context, AsyncSnapshot<String> snapshot) {
+            return _LogText(text: snapshot.data);
+          },
+        ),
+      ),
+    );
+  }
+
+  /// Open a prompt to share the log's contents using the local operating system.
+  void _shareLog(File log) async {
+    // Need to provide this position for sharing to work on iPad
+    final renderBox = shareButtonKey.currentContext?.findRenderObject() as RenderBox?;
+    if (renderBox == null) {
+      return;
+    }
+
+    final sharePosition = renderBox.localToGlobal(Offset.zero) & renderBox.size;
+
+    Share.share(
+      await log.readAsString(),
+      subject: log.uri.pathSegments.last,
+      sharePositionOrigin: sharePosition,
+    );
+  }
+}
+
+/// Widget that displays text from a log.
+class _LogText extends StatefulWidget {
+  const _LogText({Key? key, required this.text}) : super(key: key);
+
+  /// Text to display.
+  final String? text;
+
+  @override
+  State<_LogText> createState() => _LogTextState();
+}
+
+class _LogTextState extends State<_LogText> {
+  /// Vertical scroll controller
+  final _verticalScroll = ScrollController();
+
+  /// Horizontal scroll controller
+  final _horizontalScroll = ScrollController();
+
+  /// Notifier for when the text is changed
+  final _textChangeNotifier = ExternalNotifier();
+
+  /// Text style to use for log text
+  late TextStyle _textStyle;
+
+  /// List of lines to display in the text
+  List<String> _lines = [];
+
+  /// Width of the text scrolling area
+  double? _textWidth;
+
+  @override
+  void didChangeDependencies() {
+    super.didChangeDependencies();
+
+    _textStyle = Theme.of(context).textTheme.bodyMedium!.copyWith(
+          fontFamily: 'Droid Sans Mono',
+          fontSize: 12,
+        );
+
+    _updateText();
+  }
+
+  @override
+  void didUpdateWidget(covariant _LogText oldWidget) {
+    super.didUpdateWidget(oldWidget);
+
+    if (oldWidget.text != widget.text) {
+      _updateText();
+
+      WidgetsBinding.instance.addPostFrameCallback((_) {
+        _verticalScroll.jumpTo(_verticalScroll.initialScrollOffset);
+        _horizontalScroll.jumpTo(_horizontalScroll.initialScrollOffset);
+
+        _textChangeNotifier.notifyListeners();
+      });
+    }
+  }
+
+  @override
+  Widget build(BuildContext context) {
+    if (_textWidth == null) {
+      return const _SettingsLogViewLoadingSpinner();
+    }
+
+    final scrollArea = MediaQuery.removePadding(
+      context: context,
+      removeTop: true,
+      removeBottom: true,
+      child: DefaultTextStyle(
+        style: _textStyle,
+        child: EpicScrollBar(
+          padding: EdgeInsets.only(left: 16, right: 8, top: 8, bottom: 8),
+          controller: _verticalScroll,
+          // Listen to the horizontal scroll view, nested 1 level past the nearest view
+          notification: (notification) => notification.depth == 1,
+          child: Padding(
+            padding: const EdgeInsets.only(right: 8),
+            child: EpicScrollBar(
+              padding: EdgeInsets.only(right: 8),
+              controller: _horizontalScroll,
+              child: SingleChildScrollView(
+                padding: EdgeInsets.only(bottom: 16),
+                scrollDirection: Axis.horizontal,
+                controller: _horizontalScroll,
+                child: SizedBox(
+                  width: _textWidth,
+                  child: ScrollConfiguration(
+                    behavior: ScrollConfiguration.of(context).copyWith(scrollbars: false),
+
+                    // Display text as a ListView of lines for better performance. Rather than laying out all the text,
+                    // this only lays out and renders the visible lines.
+                    child: ListView.builder(
+                      shrinkWrap: true,
+                      scrollDirection: Axis.vertical,
+                      controller: _verticalScroll,
+                      itemCount: _lines.length,
+                      itemBuilder: (_, final int lineIndex) => Text(
+                        _lines[lineIndex],
+                        softWrap: false,
+                      ),
+                      prototypeItem: const Text(''),
+                    ),
+                  ),
+                ),
+              ),
+            ),
+          ),
+        ),
+      ),
+    );
+
+    // Add buttons to quickly scroll to top/bottom
+    return Stack(
+      clipBehavior: Clip.none,
+      children: [
+        scrollArea,
+        Padding(
+          padding: EdgeInsets.symmetric(vertical: 10),
+          child: Stack(
+            clipBehavior: Clip.none,
+            children: [
+              Align(
+                alignment: Alignment.topCenter,
+                child: _ScrollButton(
+                  bScrollToStart: true,
+                  scrollController: _verticalScroll,
+                  resetNotifier: _textChangeNotifier,
+                ),
+              ),
+              Align(
+                alignment: Alignment.bottomCenter,
+                child: _ScrollButton(
+                  bScrollToStart: false,
+                  scrollController: _verticalScroll,
+                  resetNotifier: _textChangeNotifier,
+                ),
+              ),
+            ],
+          ),
+        ),
+      ],
+    );
+  }
+
+  /// Update information based on the text to display.
+  void _updateText() async {
+    if (widget.text == null) {
+      _lines = const [];
+      _textWidth = null;
+      return;
+    }
+
+    _lines = widget.text!.split('\n').where((line) => line.length > 0).toList(growable: false);
+
+    // Find the width of the longest line to determine our horizontal scroll area.
+    // This only works if we use a monospace font, but is much cheaper than measuring the entire text for long logs.
+    final String longestLine = _lines.fold(
+      '',
+      (String longest, String line) => line.length > longest.length ? line : longest,
+    );
+
+    final textPainter = TextPainter(
+      textDirection: Directionality.of(context),
+      text: TextSpan(
+        text: longestLine,
+        style: _textStyle,
+      ),
+    );
+
+    textPainter.layout();
+
+    _textWidth = textPainter.width;
+  }
+}
+
+/// A button that lets the user instantly scroll to the start or end and hides itself when not relevant.
+class _ScrollButton extends StatefulWidget {
+  const _ScrollButton({
+    Key? key,
+    required this.scrollController,
+    required this.bScrollToStart,
+    this.resetNotifier,
+  }) : super(key: key);
+
+  /// Controller used to scroll the controlled view.
+  final ScrollController scrollController;
+
+  /// Whether this should scroll to the start (true) or end (false) of the scroll view.
+  final bool bScrollToStart;
+
+  /// If provided, update the button's visibility when this notifier fires (e.g. scroll bounds have changed).
+  final ChangeNotifier? resetNotifier;
+
+  @override
+  State<_ScrollButton> createState() => _ScrollButtonState();
+}
+
+class _ScrollButtonState extends State<_ScrollButton> {
+  bool _bIsVisible = false;
+  bool _bIsAnimating = false;
+
+  @override
+  void initState() {
+    super.initState();
+    widget.scrollController.addListener(_updateVisibility);
+    widget.resetNotifier?.addListener(_updateVisibility);
+
+    WidgetsBinding.instance.addPostFrameCallback((_) {
+      _updateVisibility();
+    });
+  }
+
+  @override
+  void dispose() {
+    super.dispose();
+
+    widget.scrollController.removeListener(_updateVisibility);
+    widget.resetNotifier?.removeListener(_updateVisibility);
+  }
+
+  @override
+  void didUpdateWidget(covariant _ScrollButton oldWidget) {
+    super.didUpdateWidget(oldWidget);
+
+    oldWidget.resetNotifier?.removeListener(_updateVisibility);
+    oldWidget.scrollController.removeListener(_updateVisibility);
+
+    widget.resetNotifier?.addListener(_updateVisibility);
+    widget.scrollController.addListener(_updateVisibility);
+  }
+
+  @override
+  Widget build(BuildContext context) {
+    return AnimatedOpacity(
+      opacity: (_bIsVisible && !_bIsAnimating) ? 1 : 0,
+      duration: Duration(milliseconds: 150),
+      child: FloatingActionButton(
+        child: AssetIcon(
+          path: widget.bScrollToStart ? 'assets/images/icons/chevron_up.svg' : 'assets/images/icons/chevron_down.svg',
+          size: 32,
+        ),
+        onPressed: _onPressed,
+        heroTag: null,
+      ),
+    );
+  }
+
+  /// Update whether to show the button based on scroll position.
+  void _updateVisibility() {
+    bool bNewIsVisible = false;
+    final double? scrollTarget = _getScrollTarget();
+    if (scrollTarget == null) {
+      bNewIsVisible = false;
+    } else {
+      bNewIsVisible = widget.scrollController.offset != scrollTarget;
+    }
+
+    if (bNewIsVisible != _bIsVisible) {
+      setState(() {
+        _bIsVisible = bNewIsVisible;
+      });
+    }
+  }
+
+  /// Called when the button is pressed.
+  void _onPressed() {
+    final double? scrollTarget = _getScrollTarget();
+    if (scrollTarget != null) {
+      setState(() {
+        _bIsAnimating = true;
+      });
+
+      widget.scrollController
+          .animateTo(
+            scrollTarget,
+            duration: Duration(milliseconds: 300),
+            curve: Curves.easeInOutCubic,
+          )
+          .then((_) => setState(() {
+                _bIsAnimating = false;
+              }));
+    }
+  }
+
+  /// Get the position to scroll to when the button is pressed.
+  double? _getScrollTarget() {
+    final scrollPosition = widget.scrollController.position;
+    if (!scrollPosition.hasContentDimensions) {
+      return null;
+    }
+
+    return widget.bScrollToStart ? scrollPosition.minScrollExtent : scrollPosition.maxScrollExtent;
+  }
+}
+
+/// Contents to display while the log view is loading.
+class _SettingsLogViewLoadingSpinner extends StatelessWidget {
+  const _SettingsLogViewLoadingSpinner({Key? key}) : super(key: key);
+
+  @override
+  Widget build(BuildContext context) {
+    return Center(
+      child: const Center(
+        child: SizedBox.square(
+          dimension: 120,
+          child: CircularProgressIndicator(
+            strokeWidth: 8,
+          ),
+        ),
+      ),
+    );
+  }
+}
diff --git a/Engine/Extras/VirtualProduction/EpicStageApp/lib/widgets/screens/main/toolbar/settings/pages/settings_root_actor_picker.dart b/Engine/Extras/VirtualProduction/EpicStageApp/lib/widgets/screens/main/toolbar/settings/pages/settings_root_actor_picker.dart
new file mode 100644
--- /dev/null
+++ b/Engine/Extras/VirtualProduction/EpicStageApp/lib/widgets/screens/main/toolbar/settings/pages/settings_root_actor_picker.dart
@@ -0,0 +1,64 @@
+// Copyright Epic Games, Inc. All Rights Reserved.
+
+import 'package:flutter/material.dart';
+import 'package:flutter_gen/gen_l10n/app_localizations.dart';
+import 'package:provider/provider.dart';
+
+import '../../../../../../models/settings/selected_actor_settings.dart';
+import '../../../../../../models/unreal_actor_manager.dart';
+import '../../../../../../models/unreal_types.dart';
+import '../../../../../../utilities/constants.dart';
+import '../../../../../../utilities/guarded_refresh_state.dart';
+import '../settings_generic.dart';
+
+/// Page showing the list of nDisplay root actors to pick from.
+class SettingsDialogRootActorPicker extends StatefulWidget {
+  const SettingsDialogRootActorPicker({Key? key}) : super(key: key);
+
+  static const String route = '/select_root_actor';
+
+  @override
+  State<SettingsDialogRootActorPicker> createState() => _SettingsDialogRootActorPickerState();
+}
+
+class _SettingsDialogRootActorPickerState extends State<SettingsDialogRootActorPicker> with GuardedRefreshState {
+  late final UnrealActorManager _actorManager;
+
+  @override
+  void initState() {
+    super.initState();
+
+    _actorManager = Provider.of<UnrealActorManager>(context, listen: false);
+    _actorManager.watchClassName(nDisplayRootActorClassName, refreshOnData);
+  }
+
+  @override
+  void dispose() {
+    _actorManager.stopWatchingClassName(nDisplayRootActorClassName, refreshOnData);
+    super.dispose();
+  }
+
+  @override
+  Widget build(BuildContext context) {
+    final selectedActorSettings = Provider.of<SelectedActorSettings>(context, listen: false);
+
+    return SettingsPageScaffold(
+      title: AppLocalizations.of(context)!.settingsDialogRootActorSettingLabel,
+      body: StreamBuilder(
+        stream: selectedActorSettings.displayClusterRootPath,
+        builder: (_, final AsyncSnapshot<String> rootActorPath) => Column(
+          mainAxisSize: MainAxisSize.min,
+          children: [
+            for (final UnrealObject actor in _actorManager.getActorsOfClass(nDisplayRootActorClassName))
+              SettingsMenuItem(
+                title: actor.name,
+                iconPath: 'assets/images/icons/ndisplay.svg',
+                trailingIconPath: rootActorPath == actor.path ? 'assets/images/icons/check.svg' : null,
+                onTap: () => selectedActorSettings.displayClusterRootPath.setValue(actor.path),
+              ),
+          ],
+        ),
+      ),
+    );
+  }
+}
diff --git a/Engine/Extras/VirtualProduction/EpicStageApp/lib/widgets/screens/main/toolbar/settings/settings_dialog.dart b/Engine/Extras/VirtualProduction/EpicStageApp/lib/widgets/screens/main/toolbar/settings/settings_dialog.dart
new file mode 100644
--- /dev/null
+++ b/Engine/Extras/VirtualProduction/EpicStageApp/lib/widgets/screens/main/toolbar/settings/settings_dialog.dart
@@ -0,0 +1,83 @@
+// Copyright Epic Games, Inc. All Rights Reserved.
+
+import 'package:flutter/material.dart';
+
+import '../../../../../models/navigator_keys.dart';
+import '../../../../elements/modal.dart';
+import 'pages/settings_dialog_main.dart';
+import 'pages/settings_log_list.dart';
+import 'pages/settings_log_view.dart';
+import 'pages/settings_root_actor_picker.dart';
+
+/// Dialog containing the app's settings menu.
+class SettingsDialog extends StatefulWidget {
+  const SettingsDialog({Key? key}) : super(key: key);
+
+  static void show() {
+    final route = GenericModalDialogRoute(
+      builder: (_) => SettingsDialog(),
+    );
+
+    Navigator.of(rootNavigatorKey.currentContext!, rootNavigator: true).push(route);
+  }
+
+  @override
+  State<StatefulWidget> createState() => _SettingsDialogState();
+}
+
+class _SettingsDialogState extends State<SettingsDialog> with RouteAware {
+  /// Key corresponding to the inner navigator of the settings menu.
+  final _innerNavigatorKey = GlobalKey<NavigatorState>();
+
+  @override
+  Widget build(BuildContext context) {
+    return ModalDialogCard(
+      color: Theme.of(context).colorScheme.surface,
+      shape: RoundedRectangleBorder(
+        borderRadius: BorderRadius.circular(16),
+      ),
+      child: ConstrainedBox(
+        constraints: BoxConstraints(
+          maxWidth: 600,
+          minHeight: 250,
+        ),
+        child: IntrinsicHeight(
+          child: Navigator(
+            key: _innerNavigatorKey,
+            onGenerateRoute: (final RouteSettings settings) {
+              late final Widget page;
+
+              switch (settings.name) {
+                case SettingsDialogMain.route:
+                  page = const SettingsDialogMain();
+                  break;
+
+                case SettingsDialogRootActorPicker.route:
+                  page = const SettingsDialogRootActorPicker();
+                  break;
+
+                case SettingsLogList.route:
+                  page = const SettingsLogList();
+                  break;
+
+                case SettingsLogView.route:
+                  page = const SettingsLogView();
+                  break;
+
+                default:
+                  throw Exception('No settings route named ${settings.name}');
+              }
+
+              return PageRouteBuilder(
+                settings: settings,
+                transitionDuration: Duration.zero,
+                reverseTransitionDuration: Duration.zero,
+                pageBuilder: (_, __, ___) => page,
+              );
+            },
+          ),
+        ),
+      ),
+    );
+  }
+}
diff --git a/Engine/Extras/VirtualProduction/EpicStageApp/lib/widgets/screens/main/toolbar/settings/settings_generic.dart b/Engine/Extras/VirtualProduction/EpicStageApp/lib/widgets/screens/main/toolbar/settings/settings_generic.dart
new file mode 100644
--- /dev/null
+++ b/Engine/Extras/VirtualProduction/EpicStageApp/lib/widgets/screens/main/toolbar/settings/settings_generic.dart
@@ -0,0 +1,166 @@
+// Copyright Epic Games, Inc. All Rights Reserved.
+
+import 'package:flutter/material.dart';
+
+import '../../../../../models/navigator_keys.dart';
+import '../../../../../utilities/unreal_colors.dart';
+import '../../../../elements/asset_icon.dart';
+import '../../../../elements/epic_icon_button.dart';
+import 'settings_dialog.dart';
+
+/// Outer scaffold to use for all settings pages.
+class SettingsPageScaffold extends StatelessWidget {
+  const SettingsPageScaffold({
+    Key? key,
+    required this.title,
+    required this.body,
+    this.titleBarTrailing,
+  }) : super(key: key);
+
+  /// The title of the page.
+  final String title;
+
+  /// The main contents of the page.
+  final Widget body;
+
+  /// Widget to add trailing the title bar.
+  final Widget? titleBarTrailing;
+
+  @override
+  Widget build(BuildContext context) {
+    final bool bCanPop = Navigator.of(context).canPop();
+
+    return Column(
+      mainAxisSize: MainAxisSize.min,
+      crossAxisAlignment: CrossAxisAlignment.stretch,
+      children: [
+        Container(
+          height: 40,
+          color: Theme.of(context).colorScheme.surfaceTint,
+          padding: const EdgeInsets.symmetric(horizontal: 4),
+          child: Stack(
+            alignment: Alignment.center,
+            clipBehavior: Clip.none,
+            children: [
+              Positioned(
+                left: 0,
+                child: bCanPop
+                    ? EpicIconButton(
+                        iconPath: 'assets/images/icons/chevron_left.svg',
+                        onPressed: () => Navigator.of(context).maybePop(),
+                      )
+                    : EpicIconButton(
+                        iconPath: 'assets/images/icons/close.svg',
+                        onPressed: () => Navigator.of(rootNavigatorKey.currentContext!).maybePop(),
+                      ),
+              ),
+              SizedBox(
+                width: 400,
+                child: Text(
+                  title,
+                  style: Theme.of(context).textTheme.headlineSmall!.copyWith(color: UnrealColors.white),
+                  overflow: TextOverflow.ellipsis,
+                  textAlign: TextAlign.center,
+                ),
+              ),
+              if (titleBarTrailing != null)
+                Positioned(
+                  right: 0,
+                  child: titleBarTrailing!,
+                ),
+            ],
+          ),
+        ),
+        const SettingsMenuDivider(),
+        body,
+      ],
+    );
+  }
+}
+
+/// An item in a [SettingsDialog] menu.
+class SettingsMenuItem extends StatelessWidget {
+  const SettingsMenuItem({
+    Key? key,
+    required this.title,
+    required this.iconPath,
+    this.onTap,
+    this.trailingIconPath = 'assets/images/icons/chevron_right.svg',
+    this.trailing,
+  }) : super(key: key);
+
+  /// The title to display for this item.
+  final String title;
+
+  /// The path of the icon asset to show next to the title.
+  final String iconPath;
+
+  /// Function to call when the user taps on this item.
+  final void Function()? onTap;
+
+  /// Path of the icon to show trailing all content.
+  final String? trailingIconPath;
+
+  /// An extra widget to show before the right edge/trailing icon.
+  final Widget? trailing;
+
+  @override
+  Widget build(BuildContext context) {
+    final TextStyle textStyle = Theme.of(context).textTheme.labelMedium!;
+
+    return SizedBox(
+      height: 40,
+      child: ListTile(
+        leading: SizedBox.square(
+          dimension: 24,
+          child: AssetIcon(path: iconPath),
+        ),
+        title: Text(
+          title,
+          softWrap: false,
+          overflow: TextOverflow.ellipsis,
+          style: textStyle,
+        ),
+        trailing: Row(
+          mainAxisSize: MainAxisSize.min,
+          children: [
+            if (trailing != null)
+              Padding(
+                padding: EdgeInsets.only(right: 16),
+                child: DefaultTextStyle(
+                  style: textStyle,
+                  child: trailing!,
+                ),
+              ),
+            if (trailingIconPath != null)
+              AssetIcon(
+                path: trailingIconPath!,
+                size: 24,
+              ),
+          ],
+        ),
+        horizontalTitleGap: 16,
+        minLeadingWidth: 0,
+        contentPadding: EdgeInsets.only(
+          left: 15,
+          right: 15,
+        ),
+        visualDensity: const VisualDensity(vertical: -4),
+        onTap: onTap,
+      ),
+    );
+  }
+}
+
+/// A divider shown within a [SettingsDialog] menu.
+class SettingsMenuDivider extends StatelessWidget {
+  const SettingsMenuDivider({Key? key}) : super(key: key);
+
+  @override
+  Widget build(BuildContext context) {
+    return Container(
+      height: 2,
+      color: Theme.of(context).colorScheme.surfaceVariant,
+    );
+  }
+}
diff --git a/Engine/Extras/VirtualProduction/EpicStageApp/lib/widgets/screens/reconnect_screen.dart b/Engine/Extras/VirtualProduction/EpicStageApp/lib/widgets/screens/reconnect_screen.dart
new file mode 100644
--- /dev/null
+++ b/Engine/Extras/VirtualProduction/EpicStageApp/lib/widgets/screens/reconnect_screen.dart
@@ -0,0 +1,36 @@
+// Copyright Epic Games, Inc. All Rights Reserved.
+
+import 'package:flutter/material.dart';
+
+/// Screen shown when reconnecting to the engine.
+class ReconnectScreen extends StatelessWidget {
+  const ReconnectScreen({Key? key}) : super(key: key);
+
+  static const String title = 'Connect to Unreal Engine';
+  static const String route = '/reconnect';
+
+  @override
+  Widget build(BuildContext context) {
+    return Scaffold(
+      backgroundColor: Theme.of(context).colorScheme.background,
+      body: Center(
+        child: Column(
+          mainAxisSize: MainAxisSize.min,
+          children: [
+            Text(
+              style: Theme.of(context).textTheme.titleLarge,
+              'Reconnecting to Unreal Engine...',
+            ),
+            const SizedBox(height: 20),
+            const SizedBox.square(
+              dimension: 120,
+              child: CircularProgressIndicator(
+                strokeWidth: 8,
+              ),
+            ),
+          ],
+        ),
+      ),
+    );
+  }
+}
diff --git a/Engine/Extras/VirtualProduction/EpicStageApp/lib/widgets/search_bar.dart b/Engine/Extras/VirtualProduction/EpicStageApp/lib/widgets/search_bar.dart
new file mode 100644
--- /dev/null
+++ b/Engine/Extras/VirtualProduction/EpicStageApp/lib/widgets/search_bar.dart
@@ -0,0 +1,69 @@
+// Copyright Epic Games, Inc. All Rights Reserved.
+
+import 'package:flutter/material.dart';
+import 'package:flutter_gen/gen_l10n/app_localizations.dart';
+
+import '../utilities/unreal_colors.dart';
+import 'elements/asset_icon.dart';
+
+/// A text box used to search for things by name.
+class EpicSearchBar extends StatelessWidget {
+  const EpicSearchBar({
+    Key? key,
+    this.onChanged,
+    this.controller,
+  }) : super(key: key);
+
+  /// Callback for when the search bar's text changes.
+  final ValueChanged<String>? onChanged;
+
+  /// Controller for text editing in the search bar.
+  final TextEditingController? controller;
+
+  static const double iconSize = 24;
+  static const double iconRightPadding = 8;
+  static const double iconLeftPadding = 8;
+
+  @override
+  Widget build(BuildContext context) {
+    final border = OutlineInputBorder(borderRadius: BorderRadius.circular(21));
+
+    return SizedBox(
+      height: 36,
+      child: TextField(
+        onChanged: onChanged,
+        onTapOutside: (_) {
+          final FocusScopeNode currentFocus = FocusScope.of(context);
+
+          // Flutter will throw an exception if we don't check this first
+          if (currentFocus.hasPrimaryFocus) {
+            return;
+          }
+
+          currentFocus.focusedChild?.unfocus();
+        },
+        controller: controller,
+        style: Theme.of(context).textTheme.bodyMedium,
+        textAlignVertical: TextAlignVertical.center,
+        cursorColor: UnrealColors.white,
+        cursorWidth: 1,
+        decoration: InputDecoration(
+          hintText: AppLocalizations.of(context)!.searchBarLabel,
+          filled: true,
+          fillColor: Theme.of(context).colorScheme.background,
+          hoverColor: Theme.of(context).colorScheme.background,
+          contentPadding: const EdgeInsets.symmetric(horizontal: 4),
+          border: border,
+          enabledBorder: border.copyWith(borderSide: const BorderSide(style: BorderStyle.none)),
+          focusedBorder: border.copyWith(borderSide: BorderSide(color: Theme.of(context).colorScheme.primary)),
+          prefixIcon: Padding(
+            padding: const EdgeInsets.only(left: iconLeftPadding, right: iconRightPadding),
+            child: AssetIcon(path: 'assets/images/icons/search_small.svg', size: iconSize),
+          ),
+          prefixIconConstraints: const BoxConstraints(maxWidth: iconSize + iconLeftPadding + iconRightPadding),
+          hintStyle: Theme.of(context).textTheme.bodyMedium!.copyWith(color: UnrealColors.gray56),
+        ),
+      ),
+    );
+  }
+}
diff --git a/Engine/Extras/VirtualProduction/EpicStageApp/test/widget_test.dart b/Engine/Extras/VirtualProduction/EpicStageApp/test/widget_test.dart
new file mode 100644
--- /dev/null
+++ b/Engine/Extras/VirtualProduction/EpicStageApp/test/widget_test.dart
@@ -0,0 +1,24 @@
+// Copyright Epic Games, Inc. All Rights Reserved.
+
+import 'package:flutter_test/flutter_test.dart';
+
+void main() {
+  testWidgets('Counter increments smoke test', (WidgetTester tester) async {
+    /*
+    // Build our app and trigger a frame.
+    await tester.pumpWidget(const EpicStageApp());
+
+    // Verify that our counter starts at 0.
+    expect(find.text('0'), findsOneWidget);
+    expect(find.text('1'), findsNothing);
+
+    // Tap the '+' app_icon and trigger a frame.
+    await tester.tap(find.byIcon(Icons.add));
+    await tester.pump();
+
+    // Verify that our counter has incremented.
+    expect(find.text('0'), findsNothing);
+    expect(find.text('1'), findsOneWidget);
+    */
+  });
+}
diff --git a/Engine/Plugins/Experimental/CineCameraRigs/Source/CineCameraRigs/Private/CineCameraRigRail.cpp b/Engine/Plugins/Experimental/CineCameraRigs/Source/CineCameraRigs/Private/CineCameraRigRail.cpp
--- a/Engine/Plugins/Experimental/CineCameraRigs/Source/CineCameraRigs/Private/CineCameraRigRail.cpp
+++ b/Engine/Plugins/Experimental/CineCameraRigs/Source/CineCameraRigs/Private/CineCameraRigRail.cpp
@@ -213,9 +213,10 @@ void ACineCameraRigRail::UpdateSplineMeshMID()
 		if (IsValid(SplineMeshMaterial))
 		{
 			int32 const NumSplinePoints = RailSplineComponent->GetNumberOfSplinePoints();
-			if (SplineMeshMIDs.Num() > NumSplinePoints - 1)
+			int32 const NumNeededMIDs = FMath::Max(0, NumSplinePoints - 1);
+			if (SplineMeshMIDs.Num() > NumNeededMIDs)
 			{
-				int32 const NumToRemove = SplineMeshMIDs.Num() - (NumSplinePoints - 1);
+				int32 const NumToRemove = SplineMeshMIDs.Num() - NumNeededMIDs;
 				for (int32 Index = 0; Index < NumToRemove; ++Index)
 				{
 					SplineMeshMIDs.Pop();
@@ -223,7 +224,7 @@ void ACineCameraRigRail::UpdateSplineMeshMID()
 			}
 			else
 			{
-				int32 const NumToAdd = (NumSplinePoints - 1) - SplineMeshMIDs.Num();
+				int32 const NumToAdd = NumNeededMIDs - SplineMeshMIDs.Num();
 				for (int Index = 0; Index < NumToAdd; ++Index)
 				{
 					UMaterialInstanceDynamic* SplineMeshMID = UMaterialInstanceDynamic::Create(SplineMeshMaterial, nullptr, MakeUniqueObjectName(GetTransientPackage(), UMaterialInstanceDynamic::StaticClass(), TEXT("SplineMeshMID")));
@@ -232,7 +233,7 @@ void ACineCameraRigRail::UpdateSplineMeshMID()
 				}
 			}
 			
-			for(int32 Index = 0; Index < NumSplinePoints - 1; ++Index)
+			for(int32 Index = 0; Index < NumNeededMIDs; ++Index)
 			{
 				UMaterialInstanceDynamic* SplineMeshMID = SplineMeshMIDs[Index];
 				if (!IsValid(SplineMeshMID) || (SplineMeshMID->Parent != SplineMeshMaterial))
diff --git a/Engine/Source/Runtime/CinematicCamera/Private/CameraRig_Rail.cpp b/Engine/Source/Runtime/CinematicCamera/Private/CameraRig_Rail.cpp
--- a/Engine/Source/Runtime/CinematicCamera/Private/CameraRig_Rail.cpp
+++ b/Engine/Source/Runtime/CinematicCamera/Private/CameraRig_Rail.cpp
@@ -88,7 +88,7 @@ void ACameraRig_Rail::UpdatePreviewMeshes()
 		if (PreviewRailStaticMesh)
 		{
 			int32 const NumSplinePoints = RailSplineComponent->GetNumberOfSplinePoints();
-			int32 const NumNeededPreviewMeshes = NumSplinePoints - 1;
+			int32 const NumNeededPreviewMeshes = FMath::Max(0, NumSplinePoints - 1);
 
 			// make sure our preview mesh array is correctly sized and populated
 			{
