[UnrealStage] Remove support for deprecated listSupported property

diff --git a/Engine/Extras/VirtualProduction/EpicStageApp/lib/models/unreal_engine_beacon.dart b/Engine/Extras/VirtualProduction/EpicStageApp/lib/models/unreal_engine_beacon.dart
--- a/Engine/Extras/VirtualProduction/EpicStageApp/lib/models/unreal_engine_beacon.dart
+++ b/Engine/Extras/VirtualProduction/EpicStageApp/lib/models/unreal_engine_beacon.dart
@@ -138,21 +138,15 @@ class UnrealEngineBeacon {
     final Map<InternetAddress, String> adapterNames = {};
     final Set<InternetAddress> addresses = {};
 
-    if (NetworkInterface.listSupported) {
-      // Get the list of all available network interfaces
-      final List<NetworkInterface> interfaces = await NetworkInterface.list();
+    // Get the list of all available network interfaces
+    final List<NetworkInterface> interfaces = await NetworkInterface.list();
 
-      for (final NetworkInterface interface in interfaces) {
-        addresses.addAll(interface.addresses);
+    for (final NetworkInterface interface in interfaces) {
+      addresses.addAll(interface.addresses);
 
-        for (final InternetAddress address in interface.addresses) {
-          adapterNames[address] = interface.name;
-        }
+      for (final InternetAddress address in interface.addresses) {
+        adapterNames[address] = interface.name;
       }
-    } else {
-      // We can't get a list of interfaces, so fall back to whatever the OS gives us
-      addresses.add(InternetAddress.anyIPv4);
-      addresses.add(InternetAddress.anyIPv6);
     }
 
     // Close beacons we no longer need
