[UnrealStage] Fix connection modal failing to appear

diff --git a/Engine/Extras/VirtualProduction/EpicStageApp/lib/widgets/screens/connect/mixins/connect_mixin.dart b/Engine/Extras/VirtualProduction/EpicStageApp/lib/widgets/screens/connect/mixins/connect_mixin.dart
--- a/Engine/Extras/VirtualProduction/EpicStageApp/lib/widgets/screens/connect/mixins/connect_mixin.dart
+++ b/Engine/Extras/VirtualProduction/EpicStageApp/lib/widgets/screens/connect/mixins/connect_mixin.dart
@@ -118,54 +118,52 @@ mixin ConnectMixin<T extends StatefulWidget> on State<T> {
       bResizeToAvoidBottomInset: true,
       // To avoid confusion, only allow dismissal by pressing the Cancel button
       bIsBarrierDismissible: false,
-      builder: (_) => (context as Element).debugIsActive
-          ? ModalDialogCard(
-              child: Container(
-                width: MediaQuery.of(rootNavigatorKey.currentContext!).size.width * 0.5,
-                padding: EdgeInsets.all(10),
-                child: Column(
-                  mainAxisSize: MainAxisSize.min,
-                  children: [
-                    Row(
-                      mainAxisAlignment: MainAxisAlignment.center,
-                      mainAxisSize: MainAxisSize.min,
-                      children: [
-                        AssetIcon(
-                          path: 'assets/images/icons/unreal_u_logo.svg',
-                          size: 24,
-                        ),
-                        SizedBox(width: 16),
-                        Text(
-                          AppLocalizations.of(context)!.connectScreenConnectDialogTitle,
-                          style: Theme.of(context).textTheme.displayLarge,
-                        ),
-                      ],
-                    ),
-                    const SizedBox(height: 16),
-                    Text(AppLocalizations.of(context)!
-                        .connectScreenConnectDialogMessage('${data.websocketAddress.address}:${data.websocketPort}')),
-                    const SizedBox(height: 32),
-                    SizedBox(
-                      child: const CircularProgressIndicator(strokeWidth: 8.5),
-                      height: 175,
-                      width: 175,
-                    ),
-                    Row(mainAxisAlignment: MainAxisAlignment.end, children: [
-                      EpicLozengeButton(
-                        label: AppLocalizations.of(context)!.menuButtonCancel,
-                        color: Theme.of(context).colorScheme.secondary,
-                        onPressed: () {
-                          Navigator.pop(context, false);
-                          _connectionManager.disconnect();
-                          setState(() => bIsConnecting = false);
-                        },
-                      ),
-                    ]),
-                  ],
-                ),
+      builder: (_) => ModalDialogCard(
+        child: Container(
+          width: MediaQuery.of(rootNavigatorKey.currentContext!).size.width * 0.5,
+          padding: EdgeInsets.all(10),
+          child: Column(
+            mainAxisSize: MainAxisSize.min,
+            children: [
+              Row(
+                mainAxisAlignment: MainAxisAlignment.center,
+                mainAxisSize: MainAxisSize.min,
+                children: [
+                  AssetIcon(
+                    path: 'assets/images/icons/unreal_u_logo.svg',
+                    size: 24,
+                  ),
+                  SizedBox(width: 16),
+                  Text(
+                    AppLocalizations.of(context)!.connectScreenConnectDialogTitle,
+                    style: Theme.of(context).textTheme.displayLarge,
+                  ),
+                ],
+              ),
+              const SizedBox(height: 16),
+              Text(AppLocalizations.of(context)!
+                  .connectScreenConnectDialogMessage('${data.websocketAddress.address}:${data.websocketPort}')),
+              const SizedBox(height: 32),
+              SizedBox(
+                child: const CircularProgressIndicator(strokeWidth: 8.5),
+                height: 175,
+                width: 175,
               ),
-            )
-          : SizedBox(),
+              Row(mainAxisAlignment: MainAxisAlignment.end, children: [
+                EpicLozengeButton(
+                  label: AppLocalizations.of(context)!.menuButtonCancel,
+                  color: Theme.of(context).colorScheme.secondary,
+                  onPressed: () {
+                    Navigator.pop(context, false);
+                    _connectionManager.disconnect();
+                    setState(() => bIsConnecting = false);
+                  },
+                ),
+              ]),
+            ],
+          ),
+        ),
+      ),
     );
 
     return Navigator.of(context).push(route);
