[UnrealStage] Grey out/hide outliner buttons when invalid

diff --git a/Engine/Extras/VirtualProduction/EpicStageApp/lib/widgets/elements/dropdown_button.dart b/Engine/Extras/VirtualProduction/EpicStageApp/lib/widgets/elements/dropdown_button.dart
--- a/Engine/Extras/VirtualProduction/EpicStageApp/lib/widgets/elements/dropdown_button.dart
+++ b/Engine/Extras/VirtualProduction/EpicStageApp/lib/widgets/elements/dropdown_button.dart
@@ -35,6 +35,7 @@ class ModalDropdownButton<ItemType> extends StatefulWidget {
     this.openRectExtraSpace = EdgeInsets.zero,
     this.bStretchDropdown = false,
     this.tooltipMessage,
+    this.bDisabled = false,
   }) : super(key: key);
 
   /// Function that builds the button based on the dropdown's [state].
@@ -53,6 +54,9 @@ class ModalDropdownButton<ItemType> extends StatefulWidget {
   /// An optional tooltip message to show when the button is hovered or long pressed.
   final String? tooltipMessage;
 
+  /// whether the button is disabled or not.
+  final bool? bDisabled;
+
   @override
   State<ModalDropdownButton<ItemType>> createState() => _ModalDropdownButtonState<ItemType>();
 }
@@ -69,7 +73,7 @@ class _ModalDropdownButtonState<ItemType> extends State<ModalDropdownButton<Item
         cursor: MaterialStateMouseCursor.clickable,
         child: GestureDetector(
           key: _dropdownOriginator,
-          onTap: _onTap,
+          onTap: widget.bDisabled == true ? null : _onTap,
           child: widget.buttonBuilder(context, ModalDropdownButtonState.closed),
         ),
       ),
diff --git a/Engine/Extras/VirtualProduction/EpicStageApp/lib/widgets/screens/main/sidebar/outliner_panel.dart b/Engine/Extras/VirtualProduction/EpicStageApp/lib/widgets/screens/main/sidebar/outliner_panel.dart
--- a/Engine/Extras/VirtualProduction/EpicStageApp/lib/widgets/screens/main/sidebar/outliner_panel.dart
+++ b/Engine/Extras/VirtualProduction/EpicStageApp/lib/widgets/screens/main/sidebar/outliner_panel.dart
@@ -30,7 +30,6 @@ import '../../../elements/dropdown_button.dart';
 import '../../../elements/dropdown_list_menu.dart';
 import '../../../elements/empty_placeholder.dart';
 import '../../../elements/epic_icon_button.dart';
-import '../../../elements/epic_list_view.dart';
 import '../../../elements/layout/card.dart';
 import '../../../elements/list_menu.dart';
 import '../../../elements/modal.dart';
@@ -195,6 +194,7 @@ class _OutlinerPanelState extends State<OutlinerPanel> {
   late final OutlinerPanelSettings _outlinerSettings;
 
   final TextEditingController _searchTextController = TextEditingController();
+  final ScrollController _actorListScrollController = ScrollController();
 
   /// Set of actors that we're listening to for any updates (e.g. to class information).
   Set<UnrealObject> _listenedActors = {};
@@ -300,39 +300,41 @@ class _OutlinerPanelState extends State<OutlinerPanel> {
                           );
                         },
                       ),
-                      ModalDropdownButton(
-                        buttonBuilder: (context, state) => Material(
-                          color: Colors.transparent,
-                          child: CardSubHeaderButton(
-                            iconPath: 'assets/images/icons/ellipsis.svg',
-                            tooltipMessage: AppLocalizations.of(context)!.moreActions,
-                            bIsToggledOn: state != ModalDropdownButtonState.closed,
-                            bIsVisualOnly: true,
-                          ),
-                        ),
-                        menuBuilder: (context, originTabBuilder) => TransientPreferenceBuilder(
-                          preference: selectedActorSettings.selectedActors,
-                          builder: (BuildContext context, Set<String> selectedActors) {
-                            final bool bHasSingleTarget = selectedActors.length == 1;
-                            final bool bShowFocus = widget.focusActor != null;
-                            late final bool bCanFocus;
-                            if (bShowFocus && bHasSingleTarget) {
-                              if (widget.canFocusActor == null) {
-                                // Assume we can focus any single target
-                                bCanFocus = true;
-                              } else {
-                                // Check that we can focus the current target
-                                final UnrealObject? singleTarget = _actorManager.getActorAtPath(selectedActors.first);
-                                bCanFocus = singleTarget != null && widget.canFocusActor!.call(singleTarget);
-                              }
+                      TransientPreferenceBuilder(
+                        preference: selectedActorSettings.selectedActors,
+                        builder: (BuildContext context, Set<String> selectedActors) {
+                          final bool bHasSingleTarget = selectedActors.length == 1;
+                          late final bool bCanFocus;
+
+                          if (widget.focusActor != null && bHasSingleTarget) {
+                            if (widget.canFocusActor == null) {
+                              // Assume we can focus on any single target
+                              bCanFocus = true;
                             } else {
-                              // No single target/no focus function, so we can never focus
-                              bCanFocus = false;
+                              // Check that we can focus the current target
+                              final UnrealObject? singleTarget = _actorManager.getActorAtPath(selectedActors.first);
+                              bCanFocus = singleTarget != null && widget.canFocusActor!.call(singleTarget);
                             }
-                            return DropDownListMenu(
+                          } else {
+                            // No single target/no focus function, so we can never focus
+                            bCanFocus = false;
+                          }
+
+                          return ModalDropdownButton(
+                            bDisabled: selectedActors.isEmpty,
+                            buttonBuilder: (context, state) => Material(
+                              color: Colors.transparent,
+                              child: CardSubHeaderButton(
+                                iconPath: 'assets/images/icons/ellipsis.svg',
+                                tooltipMessage: AppLocalizations.of(context)!.moreActions,
+                                bIsToggledOn: state != ModalDropdownButtonState.closed,
+                                bIsVisualOnly: selectedActors.isNotEmpty,
+                              ),
+                            ),
+                            menuBuilder: (context, originTabBuilder) => DropDownListMenu(
                               originTabBuilder: originTabBuilder,
                               children: [
-                                if (bShowFocus)
+                                if (bCanFocus)
                                   ListMenuSimpleItem(
                                     iconPath: 'assets/images/icons/focus.svg',
                                     title: AppLocalizations.of(context)!.outlinerFocusSelected,
@@ -358,29 +360,30 @@ class _OutlinerPanelState extends State<OutlinerPanel> {
                                     Navigator.of(context).pop();
                                   },
                                 ),
-                                ListMenuSimpleItem(
-                                  iconPath: 'assets/images/icons/edit.svg',
-                                  title: AppLocalizations.of(context)!.outlinerRenameSelected,
-                                  onTap: () {
-                                    Navigator.of(context).pop();
-                                    _showRenameSelectedActorDialog();
-                                  },
-                                  bIsEnabled: selectedActors.length == 1,
-                                ),
+                                if (bHasSingleTarget)
+                                  ListMenuSimpleItem(
+                                    iconPath: 'assets/images/icons/edit.svg',
+                                    title: AppLocalizations.of(context)!.outlinerRenameSelected,
+                                    onTap: () {
+                                      Navigator.of(context).pop();
+                                      _showRenameSelectedActorDialog();
+                                    },
+                                    bIsEnabled: selectedActors.length == 1,
+                                  ),
                               ],
-                            );
-                          },
-                        ),
+                            ),
+                          );
+                        },
                       ),
                     ],
                   ),
                 ),
 
-                //List of actors or onboarding placeholder
+                // List of actors or onboarding placeholder
                 Expanded(
                   child: _filteredActors.isNotEmpty
-                      ? EpicListView(
-                          padding: EdgeInsets.only(top: 4),
+                      ? ListView.builder(
+                          padding: cardListViewPadding,
                           itemCount: _filteredActors.length,
                           itemBuilder: (BuildContext context, int actorIndex) {
                             final UnrealObject actor = _filteredActors[actorIndex];
@@ -393,6 +396,7 @@ class _OutlinerPanelState extends State<OutlinerPanel> {
                               visibilityStateCallback: (state) => setState(() => _bIsSelectedActorsVisible = state),
                             );
                           },
+                          controller: _actorListScrollController,
                         )
                       : EmptyPlaceholder(
                           message: AppLocalizations.of(context)!.outlinerEmptyMessage,
@@ -841,20 +845,28 @@ class _LightCardVisibilityToggleState extends State<LightCardVisibilityToggle> w
   /// Whether the button shows the visibility icon or the hidden icon.
   late bool bIsVisible;
 
+  late SelectedActorSettings selectedActorSettings;
+
   @override
   void initState() {
     bIsVisible = widget.bIsVisible;
+    selectedActorSettings = Provider.of<SelectedActorSettings>(context, listen: false);
     super.initState();
   }
 
   @override
   Widget build(BuildContext context) {
     bIsVisible = _getSelectedActorsVisibility(context);
-    return CardSubHeaderButton(
-      iconPath: bIsVisible ? 'assets/images/icons/visible_in_game.svg' : 'assets/images/icons/hidden_in_game.svg',
-      tooltipMessage: AppLocalizations.of(context)!.outlinerToggleVisibility,
-      bIsToggledOn: false,
-      onPressed: () => _toggleVisibility(),
+    return TransientPreferenceBuilder(
+      preference: selectedActorSettings.selectedActors,
+      builder: (BuildContext context, Set<String> selectedActors) {
+        return CardSubHeaderButton(
+          iconPath: bIsVisible ? 'assets/images/icons/visible_in_game.svg' : 'assets/images/icons/hidden_in_game.svg',
+          tooltipMessage: AppLocalizations.of(context)!.outlinerToggleVisibility,
+          bIsToggledOn: false,
+          onPressed: selectedActors.isEmpty ? null : () => _toggleVisibility(),
+        );
+      },
     );
   }
 
