[UnrealStage] Clean up some old TODOs

diff --git a/Engine/Extras/VirtualProduction/EpicStageApp/lib/models/unreal_property_manager.dart b/Engine/Extras/VirtualProduction/EpicStageApp/lib/models/unreal_property_manager.dart
--- a/Engine/Extras/VirtualProduction/EpicStageApp/lib/models/unreal_property_manager.dart
+++ b/Engine/Extras/VirtualProduction/EpicStageApp/lib/models/unreal_property_manager.dart
@@ -487,24 +487,21 @@ class UnrealPropertyManager {
 
     property.bIsPendingUnexpose = false;
 
-    if (response.code != HttpResponseCode.ok) {
-      if (property.exposeCompleter != null) {
+    if (property.exposeCompleter != null) {
+      if (response.code == HttpResponseCode.ok) {
+        // We started tracking this property again after requesting unexpose, so we need to send a new expose request
+        // instead of removing it
+        _exposeProperty(propertyId);
+      } else {
         // Treat the property as re-exposed since we failed to remove it anyway
         property.exposeCompleter!.complete();
         property.exposeCompleter = null;
-        return;
       }
-
-      /// TODO: Handle this case better. We should try again to remove it, but maybe with exponential backoff
       return;
     }
 
-    if (property.exposeCompleter != null) {
-      // We started tracking this property again after requesting unexpose, so we need to send a new expose request
-      // instead of removing it
-      _exposeProperty(propertyId);
-      return;
-    }
+    // If we aren't intending to re-expose this property right away and got an error, the property was likely unexposed
+    // without our knowledge, so we can just forget about it
 
     _propertyIdsByLabel.remove(property.label);
     _trackedProperties.remove(propertyId);
diff --git a/Engine/Extras/VirtualProduction/EpicStageApp/lib/utilities/net_utilities.dart b/Engine/Extras/VirtualProduction/EpicStageApp/lib/utilities/net_utilities.dart
--- a/Engine/Extras/VirtualProduction/EpicStageApp/lib/utilities/net_utilities.dart
+++ b/Engine/Extras/VirtualProduction/EpicStageApp/lib/utilities/net_utilities.dart
@@ -10,6 +10,8 @@ import 'package:uuid/uuid.dart';
 
 import '../models/engine_connection.dart';
 
+final _log = Logger('NetUtilities');
+
 const int _protocolVersion = 0; // The current protocol version for beacon messages
 const String _protocolIdentifier = 'ES@p';
 
@@ -63,7 +65,7 @@ ConnectionData? getConnectionFromBeaconResponse(Datagram datagram) {
 
     connection = ConnectionData(uuid: engineUuid, name: name, websocketAddress: datagram.address, websocketPort: port);
   } catch (error) {
-    // TODO: Set up error logging so we can do something here
+    _log.warning('Failed to interpret beacon response from ${datagram.address.address}:/n$error');
   }
 
   return connection;
