[UnrealStage] Clean up unused Cupertino style files/includes

diff --git a/Engine/Extras/VirtualProduction/EpicStageApp/lib/models/color.dart b/Engine/Extras/VirtualProduction/EpicStageApp/lib/models/color.dart
--- a/Engine/Extras/VirtualProduction/EpicStageApp/lib/models/color.dart
+++ b/Engine/Extras/VirtualProduction/EpicStageApp/lib/models/color.dart
@@ -2,7 +2,7 @@
 
 import 'dart:math';
 
-import 'package:flutter/cupertino.dart';
+import 'package:flutter/material.dart';
 import 'package:vector_math/vector_math.dart';
 
 /// A color represented as floating-point values as opposed to Dart's native, int-based [Color].
diff --git a/Engine/Extras/VirtualProduction/EpicStageApp/lib/models/engine_connection.dart b/Engine/Extras/VirtualProduction/EpicStageApp/lib/models/engine_connection.dart
--- a/Engine/Extras/VirtualProduction/EpicStageApp/lib/models/engine_connection.dart
+++ b/Engine/Extras/VirtualProduction/EpicStageApp/lib/models/engine_connection.dart
@@ -5,7 +5,7 @@ import 'dart:convert';
 import 'dart:io';
 
 import 'package:crypto/crypto.dart';
-import 'package:flutter/cupertino.dart';
+import 'package:flutter/material.dart';
 import 'package:flutter_gen/gen_l10n/app_localizations.dart';
 import 'package:logging/logging.dart';
 import 'package:provider/provider.dart';
diff --git a/Engine/Extras/VirtualProduction/EpicStageApp/lib/models/navigator_keys.dart b/Engine/Extras/VirtualProduction/EpicStageApp/lib/models/navigator_keys.dart
--- a/Engine/Extras/VirtualProduction/EpicStageApp/lib/models/navigator_keys.dart
+++ b/Engine/Extras/VirtualProduction/EpicStageApp/lib/models/navigator_keys.dart
@@ -1,6 +1,6 @@
 // Copyright Epic Games, Inc. All Rights Reserved.
 
-import 'package:flutter/cupertino.dart';
+import 'package:flutter/material.dart';
 
 /// Key used to access the root navigator for the app without having to retrieve it from the BuildContext.
 final GlobalKey<NavigatorState> rootNavigatorKey = GlobalKey<NavigatorState>();
diff --git a/Engine/Extras/VirtualProduction/EpicStageApp/lib/models/unreal_engine_beacon.dart b/Engine/Extras/VirtualProduction/EpicStageApp/lib/models/unreal_engine_beacon.dart
--- a/Engine/Extras/VirtualProduction/EpicStageApp/lib/models/unreal_engine_beacon.dart
+++ b/Engine/Extras/VirtualProduction/EpicStageApp/lib/models/unreal_engine_beacon.dart
@@ -6,7 +6,7 @@ import 'dart:io';
 import 'dart:math';
 import 'dart:typed_data';
 
-import 'package:flutter/cupertino.dart';
+import 'package:flutter/material.dart';
 import 'package:flutter/foundation.dart';
 import 'package:logging/logging.dart';
 
diff --git a/Engine/Extras/VirtualProduction/EpicStageApp/lib/models/unreal_property_manager.dart b/Engine/Extras/VirtualProduction/EpicStageApp/lib/models/unreal_property_manager.dart
--- a/Engine/Extras/VirtualProduction/EpicStageApp/lib/models/unreal_property_manager.dart
+++ b/Engine/Extras/VirtualProduction/EpicStageApp/lib/models/unreal_property_manager.dart
@@ -2,7 +2,7 @@
 
 import 'dart:async';
 
-import 'package:flutter/cupertino.dart';
+import 'package:flutter/material.dart';
 import 'package:logging/logging.dart';
 import 'package:provider/provider.dart';
 
diff --git a/Engine/Extras/VirtualProduction/EpicStageApp/lib/utilities/preferences_bundle.dart b/Engine/Extras/VirtualProduction/EpicStageApp/lib/utilities/preferences_bundle.dart
--- a/Engine/Extras/VirtualProduction/EpicStageApp/lib/utilities/preferences_bundle.dart
+++ b/Engine/Extras/VirtualProduction/EpicStageApp/lib/utilities/preferences_bundle.dart
@@ -1,6 +1,6 @@
 // Copyright Epic Games, Inc. All Rights Reserved.
 
-import 'package:flutter/cupertino.dart';
+import 'package:flutter/material.dart';
 import 'package:provider/provider.dart';
 import 'package:streaming_shared_preferences/streaming_shared_preferences.dart';
 
diff --git a/Engine/Extras/VirtualProduction/EpicStageApp/lib/widgets/elements/cupertino/cupertino_checkbox.dart b/Engine/Extras/VirtualProduction/EpicStageApp/lib/widgets/elements/cupertino/cupertino_checkbox.dart
deleted file mode 100644
diff --git a/Engine/Extras/VirtualProduction/EpicStageApp/lib/widgets/elements/cupertino/cupertino_icon_button.dart b/Engine/Extras/VirtualProduction/EpicStageApp/lib/widgets/elements/cupertino/cupertino_icon_button.dart
deleted file mode 100644
diff --git a/Engine/Extras/VirtualProduction/EpicStageApp/lib/widgets/elements/cupertino/cupertino_pulldown.dart b/Engine/Extras/VirtualProduction/EpicStageApp/lib/widgets/elements/cupertino/cupertino_pulldown.dart
deleted file mode 100644
diff --git a/Engine/Extras/VirtualProduction/EpicStageApp/lib/widgets/elements/transform_gesture_detector.dart b/Engine/Extras/VirtualProduction/EpicStageApp/lib/widgets/elements/transform_gesture_detector.dart
--- a/Engine/Extras/VirtualProduction/EpicStageApp/lib/widgets/elements/transform_gesture_detector.dart
+++ b/Engine/Extras/VirtualProduction/EpicStageApp/lib/widgets/elements/transform_gesture_detector.dart
@@ -1,6 +1,6 @@
 // Copyright Epic Games, Inc. All Rights Reserved.
 
-import 'package:flutter/cupertino.dart';
+import 'package:flutter/material.dart';
 import 'package:flutter/gestures.dart';
 
 import '../../gestures/rotate_and_scale.dart';
diff --git a/Engine/Extras/VirtualProduction/EpicStageApp/lib/widgets/elements/unreal_widget_base.dart b/Engine/Extras/VirtualProduction/EpicStageApp/lib/widgets/elements/unreal_widget_base.dart
--- a/Engine/Extras/VirtualProduction/EpicStageApp/lib/widgets/elements/unreal_widget_base.dart
+++ b/Engine/Extras/VirtualProduction/EpicStageApp/lib/widgets/elements/unreal_widget_base.dart
@@ -1,7 +1,7 @@
 // Copyright Epic Games, Inc. All Rights Reserved.
 
-import 'package:flutter/cupertino.dart';
 import 'package:flutter/foundation.dart';
+import 'package:flutter/material.dart';
 import 'package:flutter_gen/gen_l10n/app_localizations.dart';
 import 'package:logging/logging.dart';
 import 'package:provider/provider.dart';
diff --git a/Engine/Extras/VirtualProduction/EpicStageApp/lib/widgets/screens/eula/eula_screen.dart b/Engine/Extras/VirtualProduction/EpicStageApp/lib/widgets/screens/eula/eula_screen.dart
--- a/Engine/Extras/VirtualProduction/EpicStageApp/lib/widgets/screens/eula/eula_screen.dart
+++ b/Engine/Extras/VirtualProduction/EpicStageApp/lib/widgets/screens/eula/eula_screen.dart
@@ -1,13 +1,13 @@
 // Copyright Epic Games, Inc. All Rights Reserved.
 
-import '../../elements/epic_scroll_view.dart';
-import 'package:flutter/cupertino.dart';
-import 'package:provider/provider.dart';
 import 'package:flutter/material.dart';
 import 'package:flutter_gen/gen_l10n/app_localizations.dart';
+import 'package:provider/provider.dart';
 
 import '../../../utilities/preferences_bundle.dart';
 import '../../../utilities/unreal_colors.dart';
+import '../../elements/epic_icon_button.dart';
+import '../../elements/epic_scroll_view.dart';
 import '../../elements/parsed_rich_text.dart';
 import '../connect/connect.dart';
 
@@ -46,6 +46,7 @@ class _EulaScreenState extends State<EulaScreen> {
           AppLocalizations.of(context)!.eulaAppBarTitle,
           style: Theme.of(context).textTheme.displayLarge,
         ),
+        backgroundColor: UnrealColors.gray10,
       ),
       body: EpicScrollView(
         child: Column(
@@ -70,8 +71,9 @@ class _EulaScreenState extends State<EulaScreen> {
       ),
       bottomNavigationBar: Container(
         padding: EdgeInsets.symmetric(vertical: 18, horizontal: 18),
-        decoration: BoxDecoration(color: UnrealColors.gray06),
-        child: CupertinoButton(
+        decoration: BoxDecoration(color: UnrealColors.gray10),
+        child: EpicGenericButton(
+          color: UnrealColors.highlightBlue,
           child: Text(
             bHasAcceptedEula
                 ? AppLocalizations.of(context)!.menuButtonOK
diff --git a/Engine/Extras/VirtualProduction/EpicStageApp/pubspec.yaml b/Engine/Extras/VirtualProduction/EpicStageApp/pubspec.yaml
--- a/Engine/Extras/VirtualProduction/EpicStageApp/pubspec.yaml
+++ b/Engine/Extras/VirtualProduction/EpicStageApp/pubspec.yaml
@@ -31,10 +31,6 @@ dependencies:
   flutter:
     sdk: flutter
 
-  # The following adds the Cupertino Icons font to your application.
-  # Use with the CupertinoIcons class for iOS style icons.
-  cupertino_icons: ^1.0.2
-  cupertino_list_tile: ^0.2.1
   uuid: ^3.0.0
   archive: ^3.2.0
   web_socket_channel: ^2.1.0
