[UnrealStage] If color value is ~0, force it to 1 when the user interacts with the wheel

diff --git a/Engine/Extras/VirtualProduction/EpicStageApp/lib/widgets/elements/color_grading_wheel.dart b/Engine/Extras/VirtualProduction/EpicStageApp/lib/widgets/elements/color_grading_wheel.dart
--- a/Engine/Extras/VirtualProduction/EpicStageApp/lib/widgets/elements/color_grading_wheel.dart
+++ b/Engine/Extras/VirtualProduction/EpicStageApp/lib/widgets/elements/color_grading_wheel.dart
@@ -661,8 +661,19 @@ class _DrivenColorWheelState extends State<DrivenColorWheel> with DeltaWidgetSta
       return;
     }
 
+    // Color value below which the color is effectively black
+    const double lowValue = 0.00001;
+
     final Offset scaledDeltaOffset = details.delta * deltaMultiplier / _getWheelDiameter(renderBox.constraints);
-    final List<WheelColor> deltaColors = List.filled(widget.wheelValues.length, WheelColor(scaledDeltaOffset, 0, 0));
+    final List<WheelColor> deltaColors = widget.wheelValues
+        .map((prevColor) => WheelColor(
+              scaledDeltaOffset,
+              // If the color's value is nearly 0, force it to 1 since otherwise the user won't see a color change.
+              // This imitates Unreal Editor's color editing behavior.
+              ((prevColor?.value.value ?? 0) < lowValue) ? 1 : 0,
+              0,
+            ))
+        .toList(growable: false);
     widget.onValueChanged(deltaColors);
   }
 
