[UnrealStage] Add support for zlib-compressed WebSocket traffic

diff --git a/Engine/Extras/VirtualProduction/EpicStageApp/lib/main.dart b/Engine/Extras/VirtualProduction/EpicStageApp/lib/main.dart
--- a/Engine/Extras/VirtualProduction/EpicStageApp/lib/main.dart
+++ b/Engine/Extras/VirtualProduction/EpicStageApp/lib/main.dart
@@ -15,6 +15,7 @@ import 'models/engine_connection.dart';
 import 'models/engine_passphrase_manager.dart';
 import 'models/navigator_keys.dart';
 import 'models/preview_render_manager.dart';
+import 'models/settings/connection_settings.dart';
 import 'models/settings/delta_widget_settings.dart';
 import 'models/settings/main_screen_settings.dart';
 import 'models/settings/recent_actor_settings.dart';
@@ -221,6 +222,7 @@ class _EpicStageAppState extends State<EpicStageApp> with TickerProviderStateMix
       providers: [
         Provider<EnginePassphraseManager>(create: (_) => EnginePassphraseManager()),
         Provider<PreferencesBundle>(create: (_) => _preferenceBundle),
+        Provider<ConnectionSettings>(create: (_) => ConnectionSettings(_preferenceBundle)),
         Provider<SelectedActorSettings>(create: (_) => SelectedActorSettings(_preferenceBundle)),
         Provider<StageMapSettings>(create: (_) => StageMapSettings(_preferenceBundle)),
         Provider<RecentActorSettings>(create: (_) => RecentActorSettings(_preferenceBundle)),
diff --git a/Engine/Extras/VirtualProduction/EpicStageApp/lib/models/api_version.dart b/Engine/Extras/VirtualProduction/EpicStageApp/lib/models/api_version.dart
--- a/Engine/Extras/VirtualProduction/EpicStageApp/lib/models/api_version.dart
+++ b/Engine/Extras/VirtualProduction/EpicStageApp/lib/models/api_version.dart
@@ -78,4 +78,7 @@ class EpicStageAppAPIVersion implements Comparable<EpicStageAppAPIVersion> {
 
   /// Whether the initial actor position passed to the ndisplay.preview.actor.create route is applied correctly.
   bool get bIsNewActorOverridePositionAccurate => this >= const EpicStageAppAPIVersion(1, 8, 0);
+
+  /// Whether compression of WebSocket traffic can be enabled in the engine.
+  bool get bIsWebSocketCompressionAvailable => this >= const EpicStageAppAPIVersion(1, 9, 0);
 }
diff --git a/Engine/Extras/VirtualProduction/EpicStageApp/lib/models/engine_connection.dart b/Engine/Extras/VirtualProduction/EpicStageApp/lib/models/engine_connection.dart
--- a/Engine/Extras/VirtualProduction/EpicStageApp/lib/models/engine_connection.dart
+++ b/Engine/Extras/VirtualProduction/EpicStageApp/lib/models/engine_connection.dart
@@ -3,6 +3,7 @@
 import 'dart:async';
 import 'dart:convert';
 import 'dart:io';
+import 'dart:typed_data';
 
 import 'package:crypto/crypto.dart';
 import 'package:flutter/material.dart';
@@ -23,6 +24,7 @@ import '../widgets/screens/main/stage_app_main_screen.dart';
 import '../widgets/screens/reconnect_screen.dart';
 import 'api_version.dart';
 import 'engine_passphrase_manager.dart';
+import 'settings/connection_settings.dart';
 
 final _log = Logger('EngineConnection');
 
@@ -46,6 +48,9 @@ enum EngineConnectionResult {
   /// The connection failed because the passphrase was incorrect.
   passphraseRejected,
 
+  /// The connection failed because we were unable to configure compression.
+  compressionFailed,
+
   /// The connection failed for unknown reasons.
   genericFailure,
 }
@@ -92,7 +97,8 @@ class UnrealHttpRequestWithCallback {
 
 /// Holds the state of the app's connection to the engine and notifies when the connection changes.
 class EngineConnectionManager with WidgetsBindingObserver {
-  EngineConnectionManager(this.context) {
+  EngineConnectionManager(this.context)
+      : _connectionSettings = Provider.of<ConnectionSettings>(context, listen: false) {
     _loadLastConnectionIfNone();
 
     WidgetsBinding.instance.addObserver(this);
@@ -113,6 +119,9 @@ class EngineConnectionManager with WidgetsBindingObserver {
   /// The websocket channel we currently have open to the engine.
   IOWebSocketChannel? _webSocketChannel;
 
+  /// The app's connection-related settings.
+  final ConnectionSettings _connectionSettings;
+
   /// Map from message type to callback functions for when the message type is received.
   final Map<String, List<WebSocketMessageCallback>> _messageCallbacks = {};
 
@@ -123,7 +132,11 @@ class EngineConnectionManager with WidgetsBindingObserver {
   final Map<int, Completer<UnrealHttpResponse>> _httpResponseCompleters = {};
 
   /// List of HTTP request IDs that must be processed before connection is complete.
-  final Set<int> _handshakeMessageIds = {};
+  final Set<int> _handshakeHttpMessageIds = {};
+
+  /// Map from message type to callback functions for when the message type is received during the initial connection
+  /// handshake. These messages will not be sent outside theis class.
+  final Map<String, WebSocketMessageCallback> _handshakeMessageCallbacks = {};
 
   /// List of messages that have been received, but are waiting to dispatch after connection setup is finished.
   final List<dynamic> _initialPendingMessages = [];
@@ -157,6 +170,9 @@ class EngineConnectionManager with WidgetsBindingObserver {
   /// passphrase was given.
   String? _passphraseHash;
 
+  /// The compression mode used to compress and decompress WebSocket messages.
+  WebSocketCompressionMode _compressionMode = WebSocketCompressionMode.none;
+
   /// Get the current state of the connection to the engine.
   EngineConnectionState get connectionState => _internalConnectionState;
 
@@ -220,7 +236,8 @@ class EngineConnectionManager with WidgetsBindingObserver {
     try {
       _pendingConnectionAttempt = _internalConnect(connectionData);
       result = Future.value(await _pendingConnectionAttempt);
-    } catch (error) {
+    } catch (error, stack) {
+      _log.warning('Connection failed: $error\n$stack');
       result = Future.value(EngineConnectionResult.genericFailure);
     }
 
@@ -251,13 +268,14 @@ class EngineConnectionManager with WidgetsBindingObserver {
 
       _webSocketChannel!.sink.close();
       _connectionSubscription?.cancel();
-      _onWebSocketStreamClosed(
-        bWasExpected: true,
-        // If this was in the middle of attempting to connect, we were already on the connect screen and it can handle
-        // the navigator by itself
-        bShouldReturnToConnectScreen: _pendingConnectionAttempt == null,
-      );
     }
+
+    _onWebSocketStreamClosed(
+      bWasExpected: true,
+      // If this was in the middle of attempting to connect, we were already on the connect screen and it can handle
+      // the navigator by itself
+      bShouldReturnToConnectScreen: _pendingConnectionAttempt == null,
+    );
   }
 
   /// Send a message via the current WebSocket connection in standard Unreal WebSocket format.
@@ -295,7 +313,32 @@ class EngineConnectionManager with WidgetsBindingObserver {
       message['Passphrase'] = _passphraseHash;
     }
 
-    _webSocketChannel?.sink.add(jsonEncode(message));
+    dynamic outData = jsonEncode(message);
+
+    switch (_compressionMode) {
+      case WebSocketCompressionMode.zlib:
+        final List<int> uncompressedData = utf8.encode(outData);
+        final List<int> compressedData = zlib.encode(uncompressedData);
+
+        // Unreal requires the total WebSocket message length when sending in binary mode, and also the total
+        // uncompressed size when sending compressed data.
+        const int headerSize = 8;
+        final int messageSize = compressedData.length + headerSize;
+        final outBytes = ByteData(messageSize);
+        outBytes.setUint32(0, messageSize - 4, Endian.little); // Don't include the size of this size
+        outBytes.setInt32(4, uncompressedData.length, Endian.little);
+
+        final Uint8List outBuffer = outBytes.buffer.asUint8List();
+        outBuffer.setRange(headerSize, messageSize, compressedData);
+
+        outData = outBuffer;
+        break;
+
+      default:
+        break;
+    }
+
+    _webSocketChannel?.sink.add(outData);
   }
 
   /// Send an HTTP request via the current WebSocket connection.
@@ -437,7 +480,7 @@ class EngineConnectionManager with WidgetsBindingObserver {
     _lastConnectionData = connectionData;
     _webSocketChannel = null;
     _initialPendingMessages.clear();
-    _handshakeMessageIds.clear();
+    _handshakeHttpMessageIds.clear();
 
     final String address = 'ws://${connectionData.websocketAddress.address}:${connectionData.websocketPort.toString()}';
     WebSocket? webSocket;
@@ -493,7 +536,10 @@ class EngineConnectionManager with WidgetsBindingObserver {
     _saveLastConnection(connectionData);
 
     // Check passphrase
-    final bool bIsPassphraseValid = await _checkPassphrase();
+    final bool bIsPassphraseValid = await _checkPassphrase().timeout(
+      Duration(seconds: 3),
+      onTimeout: () => false,
+    );
     if (!bIsPassphraseValid) {
       // The HTTP handler will have already disconnected us, so we can just return the result
       return EngineConnectionResult.passphraseRejected;
@@ -501,12 +547,36 @@ class EngineConnectionManager with WidgetsBindingObserver {
 
     // Retrieve API version
     _log.info('Retrieving API version');
-    _apiVersion = await _retrieveAPIVersion();
+    try {
+      _apiVersion = await _retrieveAPIVersion().timeout(Duration(seconds: 3));
+    } catch (e) {
+      _log.warning('Timed out waiting for API version');
+      disconnect();
+      return EngineConnectionResult.genericFailure;
+    }
     _log.info('API version: $_apiVersion');
 
+    // Bailed out early
     if (_bIsPendingConnectionCancelled) {
+      _log.info('User cancelled connection');
       disconnect();
-      return EngineConnectionResult.genericFailure;
+      return EngineConnectionResult.cancelled;
+    }
+
+    // Set up compression if available and desired
+    if (_apiVersion!.bIsWebSocketCompressionAvailable &&
+        _connectionSettings.webSocketCompressionMode.getValue() != WebSocketCompressionMode.none) {
+      final bool isCompressionReady = await _enableWebSocketCompression().timeout(
+        Duration(seconds: 3),
+        onTimeout: () => false,
+      );
+      if (!isCompressionReady) {
+        _log.warning('Timed out on compression response');
+        disconnect();
+        return EngineConnectionResult.compressionFailed;
+      }
+    } else {
+      _log.info('No compression requested');
     }
 
     // Now that we're fully connected, handle any messages that were waiting
@@ -520,6 +590,34 @@ class EngineConnectionManager with WidgetsBindingObserver {
     return EngineConnectionResult.success;
   }
 
+  /// Enable WebSocket compression during the connection's handshake phase.
+  Future<bool> _enableWebSocketCompression() async {
+    const compressionMessage = 'CompressionChanged';
+
+    final WebSocketCompressionMode targetCompressionMode = _connectionSettings.webSocketCompressionMode.getValue();
+    final String targetCompressionModeString = targetCompressionMode.name.toUpperCase();
+
+    _log.info('Requesting compression mode $targetCompressionModeString');
+
+    final compressionReadyCompleter = Completer<bool>();
+    final handleCompressionChanged =
+        (message) => compressionReadyCompleter.complete(message['Mode'] == targetCompressionModeString);
+
+    // Send message
+    _handshakeMessageCallbacks[compressionMessage] = handleCompressionChanged;
+    sendMessage('compression.change', {'Mode': targetCompressionModeString});
+
+    // Wait for response
+    final bool bIsCorrectMode = await compressionReadyCompleter.future;
+    _handshakeMessageCallbacks.remove(compressionMessage);
+
+    if (bIsCorrectMode) {
+      _compressionMode = targetCompressionMode;
+    }
+
+    return bIsCorrectMode;
+  }
+
   /// Called when a WebSocket error occurs.
   void _onWebSocketError(dynamic error) {
     _log.severe('WebSocket error', error);
@@ -530,6 +628,7 @@ class EngineConnectionManager with WidgetsBindingObserver {
     _webSocketChannel = null;
     _connectionSubscription = null;
     _connectionState = EngineConnectionState.disconnected;
+    _compressionMode = WebSocketCompressionMode.none;
     _apiVersion = null;
     _respondingCheck = null;
     _passphraseHash = null;
@@ -550,6 +649,19 @@ class EngineConnectionManager with WidgetsBindingObserver {
 
   /// Called when a WebSocket message is received.
   void _onWebSocketMessageReceived(dynamic data) {
+    switch (_compressionMode) {
+      case WebSocketCompressionMode.zlib:
+        try {
+          data = zlib.decode(Uint8List.fromList(data));
+        } catch (_) {
+          // Data may not have been compressed if it was cheaper to send raw, so fall through to directly decoding JSON
+        }
+        break;
+
+      default:
+        break;
+    }
+
     final String stringData = String.fromCharCodes(data);
 
     final dynamic jsonMessage;
@@ -586,8 +698,14 @@ class EngineConnectionManager with WidgetsBindingObserver {
 
     // If we're in the middle of connecting, we don't want to dispatch messages to outside systems.
     // Unless the message is required for us to complete the connection handshake, queue it for later.
-    if (_pendingConnectionAttempt != null && !_handshakeMessageIds.contains(jsonMessage[_httpRequestIdFieldName])) {
+    if (_pendingConnectionAttempt != null && !_handshakeHttpMessageIds.contains(jsonMessage[_httpRequestIdFieldName])) {
       _initialPendingMessages.add(jsonMessage);
+
+      // If this is a WebSocket message that needs to be handled during the handshake, call the internal callback even
+      // if we won't call external ones
+      final WebSocketMessageCallback? handshakeCallback = _handshakeMessageCallbacks[jsonMessage['Type']];
+      handshakeCallback?.call(jsonMessage);
+
       return;
     }
 
@@ -678,13 +796,13 @@ class EngineConnectionManager with WidgetsBindingObserver {
 
   /// Send an HTTP request via the current WebSocket connection.
   /// The response's body will contain decoded JSON data.
-  /// If [bIsHandshakeMessage] is true, the request ID will be added to [_handshakeMessageIds].
+  /// If [bIsHandshakeMessage] is true, the request ID will be added to [_handshakeHttpMessageIds].
   Future<UnrealHttpResponse> _internalSendHttpRequest(UnrealHttpRequest request, {bool bIsHandshakeMessage = false}) {
     final parameters = _generateHttpRequestParameters(request);
     final completer = _makeHttpResponseCompleterFromParameters(parameters);
 
     if (bIsHandshakeMessage) {
-      _handshakeMessageIds.add(_lastRequestId);
+      _handshakeHttpMessageIds.add(_lastRequestId);
     }
 
     sendMessage('http', parameters);
diff --git a/Engine/Extras/VirtualProduction/EpicStageApp/lib/models/settings/connection_settings.dart b/Engine/Extras/VirtualProduction/EpicStageApp/lib/models/settings/connection_settings.dart
new file mode 100644
--- /dev/null
+++ b/Engine/Extras/VirtualProduction/EpicStageApp/lib/models/settings/connection_settings.dart
@@ -0,0 +1,40 @@
+// Copyright Epic Games, Inc. All Rights Reserved.
+
+import 'package:flutter/material.dart';
+import 'package:flutter_gen/gen_l10n/app_localizations.dart';
+import 'package:streaming_shared_preferences/streaming_shared_preferences.dart';
+
+import '../../utilities/preferences_bundle.dart';
+import '../../utilities/streaming_preferences_extensions.dart';
+
+/// Compression modes that can be used for WebSocket traffic.
+enum WebSocketCompressionMode {
+  none,
+  zlib,
+}
+
+/// Given a [context] with localization data, get the name corresponding to a WebSocket [compressionMode].
+String getNameForWebSocketCompressionMode(BuildContext context, WebSocketCompressionMode compressionMode) {
+  final localizations = AppLocalizations.of(context)!;
+
+  switch (compressionMode) {
+    case WebSocketCompressionMode.none:
+      return localizations.compressionModeNone;
+
+    case WebSocketCompressionMode.zlib:
+      return localizations.compressionModeZlib;
+  }
+}
+
+/// Holds user settings used to set up the connection to the engine.
+class ConnectionSettings {
+  ConnectionSettings(PreferencesBundle preferences)
+      : webSocketCompressionMode = preferences.persistent.getEnum<WebSocketCompressionMode>(
+          'connection.webSocketCompressionMode',
+          defaultValue: WebSocketCompressionMode.zlib,
+          enumValues: WebSocketCompressionMode.values,
+        );
+
+  /// Which compression mode to use for WebSocket traffic.
+  final Preference<WebSocketCompressionMode> webSocketCompressionMode;
+}
diff --git a/Engine/Extras/VirtualProduction/EpicStageApp/lib/widgets/screens/main/toolbar/settings/pages/settings_advanced_view.dart b/Engine/Extras/VirtualProduction/EpicStageApp/lib/widgets/screens/main/toolbar/settings/pages/settings_advanced_view.dart
new file mode 100644
--- /dev/null
+++ b/Engine/Extras/VirtualProduction/EpicStageApp/lib/widgets/screens/main/toolbar/settings/pages/settings_advanced_view.dart
@@ -0,0 +1,42 @@
+// Copyright Epic Games, Inc. All Rights Reserved.
+
+import 'package:flutter/material.dart';
+import 'package:flutter_gen/gen_l10n/app_localizations.dart';
+import 'package:provider/provider.dart';
+import 'package:streaming_shared_preferences/streaming_shared_preferences.dart';
+
+import '../../../../../../models/settings/connection_settings.dart';
+import '../settings_generic.dart';
+import 'settings_compression_picker.dart';
+
+/// Page showing the advanced settings.
+class SettingsAdvancedView extends StatelessWidget {
+  const SettingsAdvancedView({Key? key}) : super(key: key);
+
+  static const String route = '/advanced';
+
+  Widget build(BuildContext context) {
+    return SettingsPageScaffold(
+      title: AppLocalizations.of(context)!.settingsDialogAdvancedLabel,
+      body: SizedBox(
+        height: 200,
+        child: MediaQuery.removePadding(
+          removeTop: true,
+          removeBottom: true,
+          context: context,
+          child: Column(children: [
+            PreferenceBuilder(
+              preference: Provider.of<ConnectionSettings>(context).webSocketCompressionMode,
+              builder: (context, final WebSocketCompressionMode compressionMode) => SettingsMenuItem(
+                title: AppLocalizations.of(context)!.settingsDialogCompressionLabel,
+                iconPath: 'assets/images/icons/compression.svg',
+                trailing: Text(getNameForWebSocketCompressionMode(context, compressionMode)),
+                onTap: () => Navigator.of(context).pushNamed(SettingsCompressionPicker.route),
+              ),
+            ),
+          ]),
+        ),
+      ),
+    );
+  }
+}
diff --git a/Engine/Extras/VirtualProduction/EpicStageApp/lib/widgets/screens/main/toolbar/settings/pages/settings_compression_picker.dart b/Engine/Extras/VirtualProduction/EpicStageApp/lib/widgets/screens/main/toolbar/settings/pages/settings_compression_picker.dart
new file mode 100644
--- /dev/null
+++ b/Engine/Extras/VirtualProduction/EpicStageApp/lib/widgets/screens/main/toolbar/settings/pages/settings_compression_picker.dart
@@ -0,0 +1,61 @@
+// Copyright Epic Games, Inc. All Rights Reserved.
+
+import 'package:flutter/material.dart';
+import 'package:flutter_gen/gen_l10n/app_localizations.dart';
+import 'package:provider/provider.dart';
+import 'package:streaming_shared_preferences/streaming_shared_preferences.dart';
+
+import '../../../../../../models/engine_connection.dart';
+import '../../../../../../models/settings/connection_settings.dart';
+import '../settings_generic.dart';
+
+/// Page showing the possible compression modes.
+class SettingsCompressionPicker extends StatelessWidget {
+  const SettingsCompressionPicker({Key? key}) : super(key: key);
+
+  static const String route = '/advanced/compression';
+
+  Widget build(BuildContext context) {
+    final ConnectionSettings connectionSettings = Provider.of<ConnectionSettings>(context);
+
+    return SettingsPageScaffold(
+      title: AppLocalizations.of(context)!.settingsDialogCompressionLabel,
+      body: SizedBox(
+        height: 200,
+        child: MediaQuery.removePadding(
+          removeTop: true,
+          removeBottom: true,
+          context: context,
+          child: Column(
+            mainAxisSize: MainAxisSize.min,
+            children: [
+              Column(
+                children: [
+                  for (final WebSocketCompressionMode mode in WebSocketCompressionMode.values)
+                    PreferenceBuilder(
+                      preference: connectionSettings.webSocketCompressionMode,
+                      builder: (context, final WebSocketCompressionMode selectedMode) => SettingsMenuItem(
+                        title: getNameForWebSocketCompressionMode(context, mode),
+                        iconPath: mode == WebSocketCompressionMode.none ? null : 'assets/images/icons/compression.svg',
+                        trailingIconPath: selectedMode == mode ? 'assets/images/icons/check.svg' : null,
+                        onTap: () => connectionSettings.webSocketCompressionMode.setValue(mode),
+                      ),
+                    ),
+                ],
+              ),
+              Spacer(),
+              if (Provider.of<EngineConnectionManager>(context).connectionState == EngineConnectionState.connected)
+                Padding(
+                  padding: EdgeInsets.only(bottom: 4),
+                  child: Text(
+                    AppLocalizations.of(context)!.reconnectSettingMessage,
+                    style: Theme.of(context).textTheme.labelMedium,
+                  ),
+                ),
+            ],
+          ),
+        ),
+      ),
+    );
+  }
+}
diff --git a/Engine/Extras/VirtualProduction/EpicStageApp/lib/widgets/screens/main/toolbar/settings/pages/settings_dialog_main.dart b/Engine/Extras/VirtualProduction/EpicStageApp/lib/widgets/screens/main/toolbar/settings/pages/settings_dialog_main.dart
--- a/Engine/Extras/VirtualProduction/EpicStageApp/lib/widgets/screens/main/toolbar/settings/pages/settings_dialog_main.dart
+++ b/Engine/Extras/VirtualProduction/EpicStageApp/lib/widgets/screens/main/toolbar/settings/pages/settings_dialog_main.dart
@@ -13,6 +13,7 @@ import '../../../../../../utilities/unreal_colors.dart';
 import '../../../../../elements/asset_icon.dart';
 import '../../../../eula/eula_screen.dart';
 import '../settings_generic.dart';
+import 'settings_advanced_view.dart';
 import 'settings_log_list.dart';
 import 'settings_root_actor_picker.dart';
 
@@ -54,6 +55,11 @@ class SettingsDialogMain extends StatelessWidget {
             iconPath: 'assets/images/icons/log.svg',
             onTap: () => Navigator.of(context).pushNamed(SettingsLogList.route),
           ),
+          SettingsMenuItem(
+            title: localizations.settingsDialogAdvancedLabel,
+            iconPath: 'assets/images/icons/advanced.svg',
+            onTap: () => Navigator.of(context).pushNamed(SettingsAdvancedView.route),
+          ),
           const SettingsMenuDivider(),
           SettingsMenuItem(
             title: localizations.settingsDialogApplicationHelpLabel,
diff --git a/Engine/Extras/VirtualProduction/EpicStageApp/lib/widgets/screens/main/toolbar/settings/pages/settings_root_actor_picker.dart b/Engine/Extras/VirtualProduction/EpicStageApp/lib/widgets/screens/main/toolbar/settings/pages/settings_root_actor_picker.dart
--- a/Engine/Extras/VirtualProduction/EpicStageApp/lib/widgets/screens/main/toolbar/settings/pages/settings_root_actor_picker.dart
+++ b/Engine/Extras/VirtualProduction/EpicStageApp/lib/widgets/screens/main/toolbar/settings/pages/settings_root_actor_picker.dart
@@ -53,7 +53,7 @@ class _SettingsDialogRootActorPickerState extends State<SettingsDialogRootActorP
               SettingsMenuItem(
                 title: actor.name,
                 iconPath: 'assets/images/icons/ndisplay.svg',
-                trailingIconPath: rootActorPath == actor.path ? 'assets/images/icons/check.svg' : null,
+                trailingIconPath: rootActorPath.data == actor.path ? 'assets/images/icons/check.svg' : null,
                 onTap: () => selectedActorSettings.displayClusterRootPath.setValue(actor.path),
               ),
           ],
diff --git a/Engine/Extras/VirtualProduction/EpicStageApp/lib/widgets/screens/main/toolbar/settings/settings_dialog.dart b/Engine/Extras/VirtualProduction/EpicStageApp/lib/widgets/screens/main/toolbar/settings/settings_dialog.dart
--- a/Engine/Extras/VirtualProduction/EpicStageApp/lib/widgets/screens/main/toolbar/settings/settings_dialog.dart
+++ b/Engine/Extras/VirtualProduction/EpicStageApp/lib/widgets/screens/main/toolbar/settings/settings_dialog.dart
@@ -4,6 +4,8 @@ import 'package:flutter/material.dart';
 
 import '../../../../../models/navigator_keys.dart';
 import '../../../../elements/modal.dart';
+import 'pages/settings_advanced_view.dart';
+import 'pages/settings_compression_picker.dart';
 import 'pages/settings_dialog_main.dart';
 import 'pages/settings_log_list.dart';
 import 'pages/settings_log_view.dart';
@@ -64,6 +66,14 @@ class _SettingsDialogState extends State<SettingsDialog> with RouteAware {
                   page = const SettingsLogView();
                   break;
 
+                case SettingsAdvancedView.route:
+                  page = const SettingsAdvancedView();
+                  break;
+
+                case SettingsCompressionPicker.route:
+                  page = const SettingsCompressionPicker();
+                  break;
+
                 default:
                   throw Exception('No settings route named ${settings.name}');
               }
diff --git a/Engine/Extras/VirtualProduction/EpicStageApp/lib/widgets/screens/main/toolbar/settings/settings_generic.dart b/Engine/Extras/VirtualProduction/EpicStageApp/lib/widgets/screens/main/toolbar/settings/settings_generic.dart
--- a/Engine/Extras/VirtualProduction/EpicStageApp/lib/widgets/screens/main/toolbar/settings/settings_generic.dart
+++ b/Engine/Extras/VirtualProduction/EpicStageApp/lib/widgets/screens/main/toolbar/settings/settings_generic.dart
@@ -83,7 +83,7 @@ class SettingsMenuItem extends StatelessWidget {
   const SettingsMenuItem({
     Key? key,
     required this.title,
-    required this.iconPath,
+    this.iconPath,
     this.onTap,
     this.trailingIconPath = 'assets/images/icons/chevron_right.svg',
     this.trailing,
@@ -92,8 +92,8 @@ class SettingsMenuItem extends StatelessWidget {
   /// The title to display for this item.
   final String title;
 
-  /// The path of the icon asset to show next to the title.
-  final String iconPath;
+  /// The path of the icon asset to show next to the title, or null to leave an empty space.
+  final String? iconPath;
 
   /// Function to call when the user taps on this item.
   final void Function()? onTap;
@@ -113,7 +113,7 @@ class SettingsMenuItem extends StatelessWidget {
       child: ListTile(
         leading: SizedBox.square(
           dimension: 24,
-          child: AssetIcon(path: iconPath),
+          child: iconPath != null ? AssetIcon(path: iconPath!) : null,
         ),
         title: Text(
           title,
