[UnrealStage] Cap Falloff Property at 1 for CCW Slider Control

diff --git a/Engine/Extras/VirtualProduction/EpicStageApp/lib/widgets/screens/main/tabs/details_tab.dart b/Engine/Extras/VirtualProduction/EpicStageApp/lib/widgets/screens/main/tabs/details_tab.dart
--- a/Engine/Extras/VirtualProduction/EpicStageApp/lib/widgets/screens/main/tabs/details_tab.dart
+++ b/Engine/Extras/VirtualProduction/EpicStageApp/lib/widgets/screens/main/tabs/details_tab.dart
@@ -369,6 +369,7 @@ class _DetailsTabState extends State<DetailsTab> with GuardedRefreshState {
       ),
       UnrealDeltaSlider(
         unrealProperties: _getPropertiesOnValidActors('Falloff'),
+        hardMax: 1,
       ),
     ]);
   }
