[UnrealStage] Allow Negative Exposure Value on Light Cards

diff --git a/Engine/Extras/VirtualProduction/EpicStageApp/lib/widgets/elements/modal.dart b/Engine/Extras/VirtualProduction/EpicStageApp/lib/widgets/elements/modal.dart
--- a/Engine/Extras/VirtualProduction/EpicStageApp/lib/widgets/elements/modal.dart
+++ b/Engine/Extras/VirtualProduction/EpicStageApp/lib/widgets/elements/modal.dart
@@ -502,20 +502,25 @@ class DoubleTextInputModalDialog extends StatelessWidget {
       parseValue: _parseDouble,
       keyboardType: keyboardType ?? TextInputType.numberWithOptions(decimal: true),
       bShowResetButton: bShowResetButton,
-      inputFormatters: inputFormatters,
+      // Additional formatter ensures we always accept numeric values, period & a negative sign (-)
+      inputFormatters: inputFormatters ?? [FilteringTextInputFormatter(RegExp('[0-9.-]'), allow: true)],
     );
   }
 
   /// Try to convert the text value to a double.
   double? _parseDouble(String text) {
-    double? result = double.tryParse(text);
-
-    if (result == null) {
-      // Try prepending a 0 so that an empty string or fractional shorthand work
-      result = double.tryParse('0$text');
+    if (text.isNotEmpty) {
+      if (text == "-" || text == "-." || text == ".") {
+        // Returning new value if text field contains only "." or "-." or ".".
+        return double.tryParse('${text}0');
+      } else {
+        // if the above condition falls through, it means we can only possibly expect a numeric value hence we parse
+        // without further manipulation/
+        return double.tryParse(text);
+      }
+    } else {
+      return 0;
     }
-
-    return result;
   }
 }
 
