[UnrealStage] Split out code and assets that can be reused in the Live Link VCAM Flutter port

diff --git a/Engine/Extras/FlutterCommon/.editorconfig b/Engine/Extras/FlutterCommon/.editorconfig
new file mode 100644
--- /dev/null
+++ b/Engine/Extras/FlutterCommon/.editorconfig
@@ -0,0 +1,11 @@
+# Copyright Epic Games, Inc. All Rights Reserved.
+
+root = true
+
+[*.dart]
+max_line_length = 120
+indent_style = space
+indent_size = 2
+tab_width = 2
+trim_trailing_whitespace = true
+insert_final_newline = true
diff --git a/Engine/Extras/FlutterCommon/.gitignore b/Engine/Extras/FlutterCommon/.gitignore
new file mode 100644
--- /dev/null
+++ b/Engine/Extras/FlutterCommon/.gitignore
@@ -0,0 +1,30 @@
+# Miscellaneous
+*.class
+*.log
+*.pyc
+*.swp
+.DS_Store
+.atom/
+.buildlog/
+.history
+.svn/
+migrate_working_dir/
+
+# IntelliJ related
+*.iml
+*.ipr
+*.iws
+.idea/
+
+# The .vscode folder contains launch configuration and tasks you configure in
+# VS Code which you may wish to be included in version control, so this line
+# is commented out by default.
+#.vscode/
+
+# Flutter/Dart/Pub related
+# Libraries should not include pubspec.lock, per https://dart.dev/guides/libraries/private-files#pubspeclock.
+/pubspec.lock
+**/doc/api/
+.dart_tool/
+.packages
+build/
diff --git a/Engine/Extras/FlutterCommon/analysis_options.yaml b/Engine/Extras/FlutterCommon/analysis_options.yaml
new file mode 100644
--- /dev/null
+++ b/Engine/Extras/FlutterCommon/analysis_options.yaml
@@ -0,0 +1,30 @@
+# This file configures the analyzer, which statically analyzes Dart code to
+# check for errors, warnings, and lints.
+#
+# The issues identified by the analyzer are surfaced in the UI of Dart-enabled
+# IDEs (https://dart.dev/tools#ides-and-editors). The analyzer can also be
+# invoked from the command line by running `flutter analyze`.
+
+# The following line activates a set of recommended lints for Flutter apps,
+# packages, and plugins designed to encourage good coding practices.
+include: package:flutter_lints/flutter.yaml
+
+linter:
+  # The lint rules applied to this project can be customized in the
+  # section below to disable rules from the `package:flutter_lints/flutter.yaml`
+  # included above or to enable additional rules. A list of all available lints
+  # and their documentation is published at
+  # https://dart-lang.github.io/linter/lints/index.html.
+  #
+  # Instead of disabling a lint rule for the entire project in the
+  # section below, it can also be suppressed for a single line of code
+  # or a specific dart file by using the `// ignore: name_of_lint` and
+  # `// ignore_for_file: name_of_lint` syntax on the line or in the file
+  # producing the lint.
+  rules:
+    prefer_relative_imports
+    # avoid_print: false  # Uncomment to disable the `avoid_print` rule
+    # prefer_single_quotes: true  # Uncomment to enable the `prefer_single_quotes` rule
+
+# Additional information about this file can be found at
+# https://dart.dev/guides/language/analysis-options
diff --git a/Engine/Extras/FlutterCommon/l10n.yaml b/Engine/Extras/FlutterCommon/l10n.yaml
new file mode 100644
--- /dev/null
+++ b/Engine/Extras/FlutterCommon/l10n.yaml
@@ -0,0 +1,7 @@
+arb-dir: lib/src/l10n
+template-arb-file: epic_common_en.arb
+synthetic-package: false
+header-file: generated_header.txt
+output-class: EpicCommonLocalizations
+output-dir: lib/src/l10n/generated
+output-localization-file: epic_common_localizations.dart
\ No newline at end of file
diff --git a/Engine/Extras/FlutterCommon/lib/src/l10n/generated_header.txt b/Engine/Extras/FlutterCommon/lib/src/l10n/generated_header.txt
new file mode 100644
--- /dev/null
+++ b/Engine/Extras/FlutterCommon/lib/src/l10n/generated_header.txt
@@ -0,0 +1,2 @@
+// Copyright Epic Games, Inc. All Rights Reserved.
+// Generated automatically by flutter_gen.
diff --git a/Engine/Extras/FlutterCommon/pubspec.yaml b/Engine/Extras/FlutterCommon/pubspec.yaml
new file mode 100644
--- /dev/null
+++ b/Engine/Extras/FlutterCommon/pubspec.yaml
@@ -0,0 +1,42 @@
+name: epic_common
+description: Utility package with shared code for use across Epic's Flutter projects.
+version: 0.0.1
+
+environment:
+  sdk: ">=3.0.6 <4.0.0"
+  flutter: "^3.7.0"
+
+dependencies:
+  flutter:
+    sdk: flutter
+
+  flutter_localizations:
+    sdk: flutter
+
+  arb_generator: ^0.0.3
+  flutter_svg: ^1.0.2
+  intl: any
+  logging: ^1.2.0
+  path_provider: ^2.0.11
+  provider: ^6.0.2
+  streaming_shared_preferences: ^2.0.0
+
+dev_dependencies:
+  flutter_test:
+    sdk: flutter
+  flutter_lints: ^2.0.0
+
+flutter:
+  generate: true
+  assets:
+    - assets/icons/
+    - assets/images/
+
+# Automatically generate arb localization files from a CSV
+arb_generator:
+  input_filepath: "lib/src/l10n/l10n.csv"
+  output_directory: "lib/src/l10n"
+  filename_prepend: "epic_common_"
+  csv_settings:
+    description_index: 1
+    base_index: 2
\ No newline at end of file
diff --git a/Engine/Extras/VirtualProduction/EpicStageApp/lib/main.dart b/Engine/Extras/VirtualProduction/EpicStageApp/lib/main.dart
--- a/Engine/Extras/VirtualProduction/EpicStageApp/lib/main.dart
+++ b/Engine/Extras/VirtualProduction/EpicStageApp/lib/main.dart
@@ -2,8 +2,12 @@
 
 import 'dart:async';
 import 'dart:isolate';
-import 'dart:ui';
 
+import 'package:epic_common/localizations.dart';
+import 'package:epic_common/logging.dart';
+import 'package:epic_common/preferences.dart';
+import 'package:epic_common/theme.dart';
+import 'package:epic_common/widgets.dart';
 import 'package:flutter/material.dart';
 import 'package:flutter/services.dart';
 import 'package:flutter_gen/gen_l10n/app_localizations.dart';
@@ -26,19 +30,13 @@ import 'models/unreal_actor_manager.dart';
 import 'models/unreal_property_manager.dart';
 import 'models/unreal_transaction_manager.dart';
 import 'routes.dart';
-import 'utilities/constants.dart';
-import 'utilities/logging.dart';
-import 'utilities/preferences_bundle.dart';
-import 'utilities/transient_preference.dart';
-import 'utilities/unreal_colors.dart';
-import 'widgets/elements/epic_scroll_view.dart';
 import 'widgets/screens/connect/connect.dart';
 import 'widgets/screens/eula/eula_screen.dart';
 
 final _log = Logger('Main');
 
 void main() async {
-  Logging.instance.initialize();
+  Logging.instance.initialize('epic_stage_app');
 
   final asyncLog = Logger('Async');
 
@@ -58,7 +56,7 @@ void initAndRunApp() async {
   final flutterLog = Logger('Flutter');
   final isolateLog = Logger('Isolate');
 
-  preloadShaders();
+  await preloadShaders();
 
   // Catch any errors originating from or otherwise passed on by Flutter
   FlutterError.onError = (FlutterErrorDetails details) {
@@ -78,8 +76,8 @@ void initAndRunApp() async {
 }
 
 /// Preload any shaders we expect to use regularly.
-void preloadShaders() async {
-  await EpicScrollView.preloadShaders();
+Future preloadShaders() async {
+  await EpicCommonWidgets.preloadShaders();
 }
 
 /// Unload any preloaded shaders.
@@ -163,61 +161,6 @@ class _EpicStageAppState extends State<EpicStageApp> with TickerProviderStateMix
       routes['/'] = routes[ConnectScreen.route]!;
     }
 
-    const colorScheme = ColorScheme.dark(
-      primary: UnrealColors.highlightBlue,
-      secondary: Color(0xff575757),
-      onPrimary: UnrealColors.white,
-      onSecondary: UnrealColors.white,
-      background: UnrealColors.gray06,
-      surfaceVariant: UnrealColors.gray10,
-      surface: UnrealColors.gray14,
-      surfaceTint: UnrealColors.gray18,
-      onSurface: UnrealColors.gray75,
-      shadow: Colors.transparent,
-    );
-
-    const textTheme = const TextTheme(
-      bodyMedium: TextStyle(
-        color: UnrealColors.gray75,
-        decorationColor: UnrealColors.gray75,
-        fontVariations: [FontVariation('wght', 400)],
-        fontSize: 14,
-      ),
-      displayLarge: TextStyle(
-        color: UnrealColors.white,
-        decorationColor: UnrealColors.white,
-        fontVariations: [FontVariation('wght', 700)],
-        fontSize: 14,
-      ),
-      displayMedium: TextStyle(
-        color: UnrealColors.white,
-        decorationColor: UnrealColors.white,
-        fontVariations: [FontVariation('wght', 400)],
-        fontSize: 14,
-      ),
-      headlineSmall: TextStyle(
-        color: UnrealColors.gray75,
-        decorationColor: UnrealColors.gray75,
-        fontVariations: [FontVariation('wght', 600)],
-        letterSpacing: 0.25,
-        fontSize: 14,
-      ),
-      titleLarge: TextStyle(
-        color: UnrealColors.gray90,
-        decorationColor: UnrealColors.gray90,
-        fontVariations: [FontVariation('wght', 700)],
-        letterSpacing: 1,
-        fontSize: 12,
-      ),
-      labelMedium: TextStyle(
-        color: UnrealColors.gray56,
-        decorationColor: UnrealColors.gray56,
-        fontVariations: [FontVariation('wght', 400)],
-        letterSpacing: 0.5,
-        fontSize: 12,
-      ),
-    );
-
     return MultiProvider(
       providers: [
         Provider<EnginePassphraseManager>(create: (_) => EnginePassphraseManager()),
@@ -245,98 +188,11 @@ class _EpicStageAppState extends State<EpicStageApp> with TickerProviderStateMix
         ),
       ],
       child: MaterialApp(
-        theme: ThemeData(
-          fontFamily: 'Inter',
-
-          textTheme: textTheme,
-          colorScheme: colorScheme,
-
-          // Disable Material "ink" splash effects
-          splashFactory: NoSplash.splashFactory,
-
-          scaffoldBackgroundColor: UnrealColors.gray10,
-          hoverColor: Colors.transparent,
-          highlightColor: Colors.transparent,
-
-          textButtonTheme: TextButtonThemeData(
-            style: ButtonStyle(
-              overlayColor: MaterialStateProperty.all(Colors.transparent),
-            ),
-          ),
-
-          checkboxTheme: CheckboxThemeData(
-            overlayColor: MaterialStateProperty.all(Colors.transparent),
-          ),
-
-          appBarTheme: const AppBarTheme(
-            shadowColor: Colors.transparent,
-          ),
-
-          tooltipTheme: const TooltipThemeData(
-            waitDuration: Duration(milliseconds: 700),
-            decoration: BoxDecoration(
-              image: DecorationImage(
-                image: AssetImage('assets/images/decoration/tooltip.png'),
-                centerSlice: Rect.fromLTRB(4, 4, 60, 60),
-              ),
-            ),
-          ),
-
-          textSelectionTheme: const TextSelectionThemeData(
-            cursorColor: UnrealColors.white,
-          ),
-
-          inputDecorationTheme: InputDecorationTheme(
-            filled: true,
-            outlineBorder: BorderSide.none,
-            border: OutlineInputBorder(
-              borderRadius: BorderRadius.circular(4),
-              borderSide: BorderSide.none,
-            ),
-            enabledBorder: OutlineInputBorder(
-              borderRadius: BorderRadius.circular(4),
-              borderSide: BorderSide.none,
-            ),
-            focusedBorder: OutlineInputBorder(
-              borderRadius: BorderRadius.circular(4),
-              borderSide: BorderSide.none,
-            ),
-            errorBorder: OutlineInputBorder(
-              borderRadius: BorderRadius.circular(4),
-              borderSide: BorderSide.none,
-            ),
-            fillColor: colorScheme.background,
-            hoverColor: colorScheme.background,
-            contentPadding: const EdgeInsets.symmetric(
-              horizontal: 12,
-            ),
-            hintStyle: textTheme.bodyMedium!.copyWith(color: UnrealColors.gray56),
-          ),
-
-          listTileTheme: const ListTileThemeData(
-            textColor: UnrealColors.gray75,
-            iconColor: UnrealColors.gray75,
-            selectedColor: UnrealColors.white,
-            selectedTileColor: UnrealColors.highlightBlue,
-          ),
-
-          scrollbarTheme: const ScrollbarThemeData(
-            thumbColor: MaterialStatePropertyAll(UnrealColors.gray22),
-            thickness: MaterialStatePropertyAll(8),
-            radius: Radius.circular(4),
-            mainAxisMargin: cardCornerRadius,
-            crossAxisMargin: cardCornerRadius,
-          ),
-
-          cardTheme: CardTheme(
-            color: UnrealColors.gray10,
-            clipBehavior: Clip.antiAlias,
-            shape: RoundedRectangleBorder(borderRadius: BorderRadius.circular(cardCornerRadius)),
-            margin: EdgeInsets.zero,
-            shadowColor: Colors.transparent,
-          ),
-        ),
-        localizationsDelegates: AppLocalizations.localizationsDelegates,
+        theme: UnrealTheme.makeThemeData(),
+        localizationsDelegates: [
+          ...AppLocalizations.localizationsDelegates,
+          ...EpicCommonLocalizations.localizationsDelegates,
+        ],
         supportedLocales: AppLocalizations.supportedLocales,
         navigatorKey: rootNavigatorKey,
         onGenerateTitle: (context) => AppLocalizations.of(context)!.appTitle,
diff --git a/Engine/Extras/VirtualProduction/EpicStageApp/lib/models/engine_connection.dart b/Engine/Extras/VirtualProduction/EpicStageApp/lib/models/engine_connection.dart
--- a/Engine/Extras/VirtualProduction/EpicStageApp/lib/models/engine_connection.dart
+++ b/Engine/Extras/VirtualProduction/EpicStageApp/lib/models/engine_connection.dart
@@ -6,6 +6,7 @@ import 'dart:io';
 import 'dart:typed_data';
 
 import 'package:crypto/crypto.dart';
+import 'package:epic_common/widgets.dart';
 import 'package:flutter/material.dart';
 import 'package:flutter_gen/gen_l10n/app_localizations.dart';
 import 'package:logging/logging.dart';
@@ -17,7 +18,6 @@ import 'package:web_socket_channel/io.dart';
 import './navigator_keys.dart';
 import '../utilities/net_utilities.dart';
 import '../utilities/unreal_utilities.dart';
-import '../widgets/elements/modal.dart';
 import '../widgets/screens/connect/connect.dart';
 import '../widgets/screens/connect/views/passphrase_modal.dart';
 import '../widgets/screens/main/stage_app_main_screen.dart';
@@ -643,7 +643,8 @@ class EngineConnectionManager with WidgetsBindingObserver {
     }
 
     if (!bWasExpected) {
-      InfoModalDialog.showInContext((context) => AppLocalizations.of(context)!.engineConnectionLostMessage);
+      final BuildContext context = rootNavigatorKey.currentContext!;
+      InfoModalDialog.show(context, AppLocalizations.of(context)!.engineConnectionLostMessage);
     }
   }
 
diff --git a/Engine/Extras/VirtualProduction/EpicStageApp/lib/models/settings/base_color_tab_settings.dart b/Engine/Extras/VirtualProduction/EpicStageApp/lib/models/settings/base_color_tab_settings.dart
--- a/Engine/Extras/VirtualProduction/EpicStageApp/lib/models/settings/base_color_tab_settings.dart
+++ b/Engine/Extras/VirtualProduction/EpicStageApp/lib/models/settings/base_color_tab_settings.dart
@@ -1,8 +1,8 @@
 // Copyright Epic Games, Inc. All Rights Reserved.
 
+import 'package:epic_common/preferences.dart';
 import 'package:streaming_shared_preferences/streaming_shared_preferences.dart';
 
-import '../../utilities/preferences_bundle.dart';
 import '../../utilities/streaming_preferences_extensions.dart';
 import 'color_wheel_settings.dart';
 
diff --git a/Engine/Extras/VirtualProduction/EpicStageApp/lib/models/settings/color_grading_tab_settings.dart b/Engine/Extras/VirtualProduction/EpicStageApp/lib/models/settings/color_grading_tab_settings.dart
--- a/Engine/Extras/VirtualProduction/EpicStageApp/lib/models/settings/color_grading_tab_settings.dart
+++ b/Engine/Extras/VirtualProduction/EpicStageApp/lib/models/settings/color_grading_tab_settings.dart
@@ -1,9 +1,8 @@
 // Copyright Epic Games, Inc. All Rights Reserved.
 
+import 'package:epic_common/preferences.dart';
 import 'package:streaming_shared_preferences/streaming_shared_preferences.dart';
 
-import '../../utilities/preferences_bundle.dart';
-import '../../utilities/transient_preference.dart';
 import '../unreal_types.dart';
 
 /// Holds user settings used in the Color Grading tab.
diff --git a/Engine/Extras/VirtualProduction/EpicStageApp/lib/models/settings/color_wheel_settings.dart b/Engine/Extras/VirtualProduction/EpicStageApp/lib/models/settings/color_wheel_settings.dart
--- a/Engine/Extras/VirtualProduction/EpicStageApp/lib/models/settings/color_wheel_settings.dart
+++ b/Engine/Extras/VirtualProduction/EpicStageApp/lib/models/settings/color_wheel_settings.dart
@@ -1,10 +1,8 @@
 // Copyright Epic Games, Inc. All Rights Reserved.
 
+import 'package:epic_common/preferences.dart';
 import 'package:streaming_shared_preferences/streaming_shared_preferences.dart';
 
-import '../../utilities/preferences_bundle.dart';
-import '../../utilities/streaming_preferences_extensions.dart';
-
 /// What type of color editor to display.
 enum ColorWheelEditMode {
   wheel,
diff --git a/Engine/Extras/VirtualProduction/EpicStageApp/lib/models/settings/connection_settings.dart b/Engine/Extras/VirtualProduction/EpicStageApp/lib/models/settings/connection_settings.dart
--- a/Engine/Extras/VirtualProduction/EpicStageApp/lib/models/settings/connection_settings.dart
+++ b/Engine/Extras/VirtualProduction/EpicStageApp/lib/models/settings/connection_settings.dart
@@ -1,12 +1,10 @@
 // Copyright Epic Games, Inc. All Rights Reserved.
 
+import 'package:epic_common/preferences.dart';
 import 'package:flutter/material.dart';
 import 'package:flutter_gen/gen_l10n/app_localizations.dart';
 import 'package:streaming_shared_preferences/streaming_shared_preferences.dart';
 
-import '../../utilities/preferences_bundle.dart';
-import '../../utilities/streaming_preferences_extensions.dart';
-
 /// Compression modes that can be used for WebSocket traffic.
 enum WebSocketCompressionMode {
   none,
diff --git a/Engine/Extras/VirtualProduction/EpicStageApp/lib/models/settings/delta_widget_settings.dart b/Engine/Extras/VirtualProduction/EpicStageApp/lib/models/settings/delta_widget_settings.dart
--- a/Engine/Extras/VirtualProduction/EpicStageApp/lib/models/settings/delta_widget_settings.dart
+++ b/Engine/Extras/VirtualProduction/EpicStageApp/lib/models/settings/delta_widget_settings.dart
@@ -1,13 +1,12 @@
 // Copyright Epic Games, Inc. All Rights Reserved.
 
+import 'package:epic_common/preferences.dart';
 import 'package:flutter/material.dart';
 import 'package:flutter_gen/gen_l10n/app_localizations.dart';
 import 'package:logging/logging.dart';
 import 'package:streaming_shared_preferences/streaming_shared_preferences.dart';
 
 import '../../utilities/debug_utilities.dart';
-import '../../utilities/preferences_bundle.dart';
-import '../../utilities/transient_preference.dart';
 import '../navigator_keys.dart';
 
 final _log = Logger('DeltaWidgetSettings');
diff --git a/Engine/Extras/VirtualProduction/EpicStageApp/lib/models/settings/details_tab_settings.dart b/Engine/Extras/VirtualProduction/EpicStageApp/lib/models/settings/details_tab_settings.dart
--- a/Engine/Extras/VirtualProduction/EpicStageApp/lib/models/settings/details_tab_settings.dart
+++ b/Engine/Extras/VirtualProduction/EpicStageApp/lib/models/settings/details_tab_settings.dart
@@ -1,10 +1,8 @@
 // Copyright Epic Games, Inc. All Rights Reserved.
 
+import 'package:epic_common/preferences.dart';
 import 'package:streaming_shared_preferences/streaming_shared_preferences.dart';
 
-import '../../utilities/preferences_bundle.dart';
-import '../../utilities/streaming_preferences_extensions.dart';
-
 /// Which type of actor properties to display in the details panel.
 enum DetailsPropertyDisplayType {
   appearance,
diff --git a/Engine/Extras/VirtualProduction/EpicStageApp/lib/models/settings/floating_map_settings.dart b/Engine/Extras/VirtualProduction/EpicStageApp/lib/models/settings/floating_map_settings.dart
--- a/Engine/Extras/VirtualProduction/EpicStageApp/lib/models/settings/floating_map_settings.dart
+++ b/Engine/Extras/VirtualProduction/EpicStageApp/lib/models/settings/floating_map_settings.dart
@@ -1,10 +1,8 @@
 // Copyright Epic Games, Inc. All Rights Reserved.
 
+import 'package:epic_common/preferences.dart';
 import 'package:streaming_shared_preferences/streaming_shared_preferences.dart';
 
-import '../../utilities/preferences_bundle.dart';
-import '../../utilities/streaming_preferences_extensions.dart';
-
 /// Which side of the screen the map preview is on.
 enum PreviewMapSide {
   left,
diff --git a/Engine/Extras/VirtualProduction/EpicStageApp/lib/models/settings/main_screen_settings.dart b/Engine/Extras/VirtualProduction/EpicStageApp/lib/models/settings/main_screen_settings.dart
--- a/Engine/Extras/VirtualProduction/EpicStageApp/lib/models/settings/main_screen_settings.dart
+++ b/Engine/Extras/VirtualProduction/EpicStageApp/lib/models/settings/main_screen_settings.dart
@@ -1,9 +1,8 @@
 // Copyright Epic Games, Inc. All Rights Reserved.
 
+import 'package:epic_common/preferences.dart';
 import 'package:streaming_shared_preferences/streaming_shared_preferences.dart';
 
-import '../../utilities/preferences_bundle.dart';
-
 /// Holds settings that affect the entire layout of the main screen.
 class MainScreenSettings {
   MainScreenSettings(PreferencesBundle preferences)
diff --git a/Engine/Extras/VirtualProduction/EpicStageApp/lib/models/settings/outliner_panel_settings.dart b/Engine/Extras/VirtualProduction/EpicStageApp/lib/models/settings/outliner_panel_settings.dart
--- a/Engine/Extras/VirtualProduction/EpicStageApp/lib/models/settings/outliner_panel_settings.dart
+++ b/Engine/Extras/VirtualProduction/EpicStageApp/lib/models/settings/outliner_panel_settings.dart
@@ -1,8 +1,8 @@
 // Copyright Epic Games, Inc. All Rights Reserved.
 
+import 'package:epic_common/preferences.dart';
 import 'package:streaming_shared_preferences/streaming_shared_preferences.dart';
 
-import '../../utilities/preferences_bundle.dart';
 import '../../utilities/streaming_preferences_extensions.dart';
 import '../unreal_object_filters.dart';
 
diff --git a/Engine/Extras/VirtualProduction/EpicStageApp/lib/models/settings/recent_actor_settings.dart b/Engine/Extras/VirtualProduction/EpicStageApp/lib/models/settings/recent_actor_settings.dart
--- a/Engine/Extras/VirtualProduction/EpicStageApp/lib/models/settings/recent_actor_settings.dart
+++ b/Engine/Extras/VirtualProduction/EpicStageApp/lib/models/settings/recent_actor_settings.dart
@@ -1,9 +1,9 @@
 // Copyright Epic Games, Inc. All Rights Reserved.
 
+import 'package:epic_common/preferences.dart';
+import 'package:epic_common/utilities/json.dart';
 import 'package:streaming_shared_preferences/streaming_shared_preferences.dart';
 
-import '../../utilities/json_utilities.dart';
-import '../../utilities/preferences_bundle.dart';
 import '../../utilities/streaming_preferences_extensions.dart';
 
 enum RecentlyPlacedActorType {
diff --git a/Engine/Extras/VirtualProduction/EpicStageApp/lib/models/settings/selected_actor_settings.dart b/Engine/Extras/VirtualProduction/EpicStageApp/lib/models/settings/selected_actor_settings.dart
--- a/Engine/Extras/VirtualProduction/EpicStageApp/lib/models/settings/selected_actor_settings.dart
+++ b/Engine/Extras/VirtualProduction/EpicStageApp/lib/models/settings/selected_actor_settings.dart
@@ -1,10 +1,8 @@
 // Copyright Epic Games, Inc. All Rights Reserved.
 
+import 'package:epic_common/preferences.dart';
 import 'package:streaming_shared_preferences/streaming_shared_preferences.dart';
 
-import '../../utilities/preferences_bundle.dart';
-import '../../utilities/transient_preference.dart';
-
 /// Holds user settings that impact the state of the selected actors.
 class SelectedActorSettings {
   SelectedActorSettings(PreferencesBundle preferences)
diff --git a/Engine/Extras/VirtualProduction/EpicStageApp/lib/models/settings/stage_map_settings.dart b/Engine/Extras/VirtualProduction/EpicStageApp/lib/models/settings/stage_map_settings.dart
--- a/Engine/Extras/VirtualProduction/EpicStageApp/lib/models/settings/stage_map_settings.dart
+++ b/Engine/Extras/VirtualProduction/EpicStageApp/lib/models/settings/stage_map_settings.dart
@@ -2,13 +2,10 @@
 
 import 'dart:ui';
 
+import 'package:epic_common/preferences.dart';
 import 'package:streaming_shared_preferences/streaming_shared_preferences.dart';
 import 'package:vector_math/vector_math_64.dart' as vec;
 
-import '../../utilities/preferences_bundle.dart';
-import '../../utilities/streaming_preferences_extensions.dart';
-import '../../utilities/transient_preference.dart';
-
 enum ProjectionMode {
   azimuthal,
   orthographic,
diff --git a/Engine/Extras/VirtualProduction/EpicStageApp/lib/models/unreal_class_data.dart b/Engine/Extras/VirtualProduction/EpicStageApp/lib/models/unreal_class_data.dart
--- a/Engine/Extras/VirtualProduction/EpicStageApp/lib/models/unreal_class_data.dart
+++ b/Engine/Extras/VirtualProduction/EpicStageApp/lib/models/unreal_class_data.dart
@@ -119,7 +119,7 @@ class UnrealClassRegistry {
   static final nDisplayRootActor = UnrealClass(
     'DisplayClusterRootActor',
     getDisplayName: (context) => AppLocalizations.of(context)!.actorNameDisplayClusterRoot,
-    iconPath: 'assets/images/icons/ndisplay.svg',
+    iconPath: 'packages/epic_common/assets/icons/ndisplay.svg',
     parent: actor,
     classNames: [nDisplayRootActorClassName],
   );
@@ -127,7 +127,7 @@ class UnrealClassRegistry {
   static final lightCard = UnrealClass(
     'LightCard',
     getDisplayName: (context) => AppLocalizations.of(context)!.actorNameLightCard,
-    iconPath: 'assets/images/icons/light_card.svg',
+    iconPath: 'packages/epic_common/assets/icons/light_card.svg',
     color: const Color(0xffc4a300),
     iconPathFunction: (UnrealObject object) {
       final lightCardData = object.getPerClassData<LightCardActorData>();
@@ -137,11 +137,11 @@ class UnrealClassRegistry {
         }
 
         if (lightCardData.bIsUV == true) {
-          return 'assets/images/icons/light_card_uv.svg';
+          return 'packages/epic_common/assets/icons/light_card_uv.svg';
         }
 
         if (lightCardData.bIsFlag == true) {
-          return 'assets/images/icons/light_card_flag.svg';
+          return 'packages/epic_common/assets/icons/light_card_flag.svg';
         }
       }
 
@@ -168,7 +168,7 @@ class UnrealClassRegistry {
   static final chromakeyCard = UnrealClass(
     'ChromakeyCard',
     getDisplayName: (context) => AppLocalizations.of(context)!.actorNameChromakeyCard,
-    iconPath: 'assets/images/icons/chromakey_card.svg',
+    iconPath: 'packages/epic_common/assets/icons/chromakey_card.svg',
     color: const Color(0xff16884a),
     parent: lightCard,
     classNames: [chromakeyCardClassName],
@@ -178,7 +178,7 @@ class UnrealClassRegistry {
     'ColorCorrectRegion',
     getDisplayName: (context) => AppLocalizations.of(context)!.actorNameColorCorrectRegion,
     color: const Color(0xffd72e80),
-    iconPath: 'assets/images/icons/color_correct_region.svg',
+    iconPath: 'packages/epic_common/assets/icons/color_correct_region.svg',
     parent: actor,
     classNames: [colorCorrectRegionClassName],
   );
@@ -186,7 +186,7 @@ class UnrealClassRegistry {
   static final colorCorrectWindow = UnrealClass(
     'ColorCorrectWindow',
     getDisplayName: (context) => AppLocalizations.of(context)!.actorNameColorCorrectWindow,
-    iconPath: 'assets/images/icons/color_correct_window.svg',
+    iconPath: 'packages/epic_common/assets/icons/color_correct_window.svg',
     parent: colorCorrectRegion,
     classNames: colorCorrectWindowClassNames,
   );
@@ -194,7 +194,7 @@ class UnrealClassRegistry {
   static final postProcessVolume = UnrealClass(
     'PostProcessVolume',
     getDisplayName: (context) => AppLocalizations.of(context)!.actorNamePostProcessVolume,
-    iconPath: 'assets/images/icons/post_process_volume.svg',
+    iconPath: 'packages/epic_common/assets/icons/post_process_volume.svg',
     parent: actor,
     classNames: [postProcessVolumeClassName],
   );
diff --git a/Engine/Extras/VirtualProduction/EpicStageApp/lib/models/unreal_object_filters.dart b/Engine/Extras/VirtualProduction/EpicStageApp/lib/models/unreal_object_filters.dart
--- a/Engine/Extras/VirtualProduction/EpicStageApp/lib/models/unreal_object_filters.dart
+++ b/Engine/Extras/VirtualProduction/EpicStageApp/lib/models/unreal_object_filters.dart
@@ -95,7 +95,7 @@ class UnrealObjectFilterRegistry {
     _CustomUnrealObjectFilter(
       internalName: 'Flag',
       getDisplayName: (context) => AppLocalizations.of(context)!.actorNameFlag,
-      iconPath: 'assets/images/icons/light_card_flag.svg',
+      iconPath: 'packages/epic_common/assets/icons/light_card_flag.svg',
       passFunction: (UnrealObject object) {
         final lightcardData = object.getPerClassData<LightCardActorData>();
         return lightcardData?.bIsFlag == true;
diff --git a/Engine/Extras/VirtualProduction/EpicStageApp/lib/utilities/constants.dart b/Engine/Extras/VirtualProduction/EpicStageApp/lib/utilities/constants.dart
--- a/Engine/Extras/VirtualProduction/EpicStageApp/lib/utilities/constants.dart
+++ b/Engine/Extras/VirtualProduction/EpicStageApp/lib/utilities/constants.dart
@@ -2,21 +2,6 @@
 
 import 'package:flutter/material.dart';
 
-/// Radius of the corners for top-level UI elements.
-const double outerCornerRadius = 16;
-
-/// Margin between the cards in a tab.
-const double cardMargin = 8;
-
-/// Radius between sections within cards.
-const double sectionMargin = 4;
-
-/// Radius of the corners for cards in a tab.
-const double cardCornerRadius = 8;
-
-/// Padding for scrollable list views within cards.
-const EdgeInsetsGeometry cardListViewPadding = EdgeInsets.only(top: 2);
-
 /// The Unreal class for nDisplay root actors.
 const String nDisplayRootActorClassName = '/Script/DisplayCluster.DisplayClusterRootActor';
 
diff --git a/Engine/Extras/VirtualProduction/EpicStageApp/lib/utilities/debug_utilities.dart b/Engine/Extras/VirtualProduction/EpicStageApp/lib/utilities/debug_utilities.dart
--- a/Engine/Extras/VirtualProduction/EpicStageApp/lib/utilities/debug_utilities.dart
+++ b/Engine/Extras/VirtualProduction/EpicStageApp/lib/utilities/debug_utilities.dart
@@ -1,9 +1,9 @@
 // Copyright Epic Games, Inc. All Rights Reserved.
 
+import 'package:epic_common/widgets.dart';
 import 'package:logging/logging.dart';
 
 import '../models/navigator_keys.dart';
-import '../widgets/elements/modal.dart';
 
 final _log = Logger('DebugAlert');
 
@@ -14,5 +14,5 @@ void showDebugAlert(String message) {
   }
 
   _log.info(message);
-  InfoModalDialog.show(message);
+  InfoModalDialog.show(rootNavigatorKey.currentContext!, message);
 }
diff --git a/Engine/Extras/VirtualProduction/EpicStageApp/lib/utilities/json_utilities.dart b/Engine/Extras/VirtualProduction/EpicStageApp/lib/utilities/json_utilities.dart
deleted file mode 100644
diff --git a/Engine/Extras/VirtualProduction/EpicStageApp/lib/utilities/logging.dart b/Engine/Extras/VirtualProduction/EpicStageApp/lib/utilities/logging.dart
deleted file mode 100644
diff --git a/Engine/Extras/VirtualProduction/EpicStageApp/lib/utilities/preferences_bundle.dart b/Engine/Extras/VirtualProduction/EpicStageApp/lib/utilities/preferences_bundle.dart
deleted file mode 100644
diff --git a/Engine/Extras/VirtualProduction/EpicStageApp/lib/utilities/preloadable_shader.dart b/Engine/Extras/VirtualProduction/EpicStageApp/lib/utilities/preloadable_shader.dart
deleted file mode 100644
diff --git a/Engine/Extras/VirtualProduction/EpicStageApp/lib/utilities/streaming_preferences_extensions.dart b/Engine/Extras/VirtualProduction/EpicStageApp/lib/utilities/streaming_preferences_extensions.dart
--- a/Engine/Extras/VirtualProduction/EpicStageApp/lib/utilities/streaming_preferences_extensions.dart
+++ b/Engine/Extras/VirtualProduction/EpicStageApp/lib/utilities/streaming_preferences_extensions.dart
@@ -8,37 +8,11 @@ import 'package:streaming_shared_preferences/streaming_shared_preferences.dart';
 
 import '../models/settings/recent_actor_settings.dart';
 import '../models/unreal_object_filters.dart';
-import 'json_utilities.dart';
 
 final _log = Logger('PrefExtensions');
 
 /// Extensions to serialize/deserialize new types of preference values.
 extension PreferenceAdapterExtensions on StreamingSharedPreferences {
-  /// Starts with the current String set value for the given [key], then emits a new value every time there are changes
-  /// to the value associated with [key].
-  ///
-  /// If the value is null, starts with the value provided in [defaultValue]. When the value transitions from non-null
-  /// to null (ie. when the value is removed), emits [defaultValue].
-  Preference<Set<String>> getStringSet(
-    String key, {
-    required Set<String> defaultValue,
-  }) {
-    return getCustomValue(key, defaultValue: defaultValue, adapter: StringSetAdapter.instance);
-  }
-
-  /// Starts with the current Enum value for the given [key], then emits a new value every time there are changes to the
-  /// value associated with [key].
-  ///
-  /// If the value is null, starts with the value provided in [defaultValue]. When the value transitions from non-null
-  /// to null (ie. when the value is removed), emits [defaultValue].
-  Preference<T> getEnum<T extends Enum>(
-    String key, {
-    required T defaultValue,
-    required List<T> enumValues,
-  }) {
-    return getCustomValue(key, defaultValue: defaultValue, adapter: EnumAdapter(enumValues));
-  }
-
   /// Starts with the current list value for the given [key], then emits a new value every time there are changes to the
   /// value associated with [key].
   ///
@@ -64,43 +38,6 @@ extension PreferenceAdapterExtensions on StreamingSharedPreferences {
   }
 }
 
-/// A [PreferenceAdapter] implementation for storing and retrieving a [Set] of [String] objects.
-class StringSetAdapter extends PreferenceAdapter<Set<String>> {
-  static const instance = StringSetAdapter._();
-  const StringSetAdapter._();
-
-  @override
-  Set<String>? getValue(preferences, key) => preferences.getStringList(key)?.toSet();
-
-  @override
-  Future<bool> setValue(preferences, key, values) => preferences.setStringList(key, values.toList(growable: false));
-}
-
-/// A [PreferenceAdapter] implementation for storing and retrieving an [Enum] value as a string.
-class EnumAdapter<T extends Enum> extends PreferenceAdapter<T> {
-  const EnumAdapter(this.enumValues);
-
-  final List<T> enumValues;
-
-  @override
-  T? getValue(preferences, key) {
-    try {
-      final String? name = preferences.getString(key);
-
-      if (name == null) {
-        return null;
-      }
-
-      return jsonToEnumValue(name, enumValues);
-    } catch (_) {
-      return null;
-    }
-  }
-
-  @override
-  Future<bool> setValue(preferences, key, value) => preferences.setString(key, enumToJsonValue(value));
-}
-
 /// A [PreferenceAdapter] implementation for storing and retrieving a [List] of [RecentlyPlacedActorData] values.
 class RecentlyPlacedActorListAdapter extends PreferenceAdapter<List<RecentlyPlacedActorData>> {
   static const instance = RecentlyPlacedActorListAdapter._();
diff --git a/Engine/Extras/VirtualProduction/EpicStageApp/lib/utilities/transient_preference.dart b/Engine/Extras/VirtualProduction/EpicStageApp/lib/utilities/transient_preference.dart
deleted file mode 100644
diff --git a/Engine/Extras/VirtualProduction/EpicStageApp/lib/utilities/unreal_colors.dart b/Engine/Extras/VirtualProduction/EpicStageApp/lib/utilities/unreal_colors.dart
deleted file mode 100644
diff --git a/Engine/Extras/VirtualProduction/EpicStageApp/lib/widgets/elements/asset_icon.dart b/Engine/Extras/VirtualProduction/EpicStageApp/lib/widgets/elements/asset_icon.dart
deleted file mode 100644
diff --git a/Engine/Extras/VirtualProduction/EpicStageApp/lib/widgets/elements/color_grading_wheel.dart b/Engine/Extras/VirtualProduction/EpicStageApp/lib/widgets/elements/color_grading_wheel.dart
--- a/Engine/Extras/VirtualProduction/EpicStageApp/lib/widgets/elements/color_grading_wheel.dart
+++ b/Engine/Extras/VirtualProduction/EpicStageApp/lib/widgets/elements/color_grading_wheel.dart
@@ -4,6 +4,9 @@ import 'dart:async';
 import 'dart:math' as math;
 
 import 'package:async/async.dart';
+import 'package:epic_common/preferences.dart';
+import 'package:epic_common/theme.dart';
+import 'package:epic_common/widgets.dart';
 import 'package:flutter/material.dart';
 import 'package:flutter_gen/gen_l10n/app_localizations.dart';
 import 'package:provider/provider.dart';
@@ -13,17 +16,11 @@ import '../../../models/property_modify_operations.dart';
 import '../../../utilities/math_utilities.dart';
 import '../../models/settings/color_wheel_settings.dart';
 import '../../models/settings/delta_widget_settings.dart';
-import '../../utilities/constants.dart';
 import '../../utilities/guarded_refresh_state.dart';
-import '../../utilities/preferences_bundle.dart';
-import '../../utilities/transient_preference.dart';
-import '../../utilities/unreal_colors.dart';
 import '../../utilities/unreal_utilities.dart';
-import 'asset_icon.dart';
 import 'color_slider.dart';
 import 'delta_slider.dart';
 import 'delta_widget_base.dart';
-import 'selector_bar.dart';
 import 'unreal_widget_base.dart';
 
 /// Diameter of a dot on the color wheel indicating a value.
@@ -577,8 +574,8 @@ class _DrivenColorWheelState extends State<DrivenColorWheel> with DeltaWidgetSta
                         icon: AssetIcon(
                           size: 20,
                           path: editMode == ColorWheelEditMode.wheel
-                              ? 'assets/images/icons/sliders.svg'
-                              : 'assets/images/icons/color_wheel.svg',
+                              ? 'packages/epic_common/assets/icons/sliders.svg'
+                              : 'packages/epic_common/assets/icons/color_wheel.svg',
                         ),
                       ),
                     ),
@@ -594,7 +591,7 @@ class _DrivenColorWheelState extends State<DrivenColorWheel> with DeltaWidgetSta
             child: widget.extraTopWidget!,
           ),
         Padding(
-          padding: EdgeInsets.symmetric(horizontal: cardMargin),
+          padding: EdgeInsets.symmetric(horizontal: UnrealTheme.cardMargin),
           child: Column(
             children: [
               // Color wheel
diff --git a/Engine/Extras/VirtualProduction/EpicStageApp/lib/widgets/elements/control_lock.dart b/Engine/Extras/VirtualProduction/EpicStageApp/lib/widgets/elements/control_lock.dart
--- a/Engine/Extras/VirtualProduction/EpicStageApp/lib/widgets/elements/control_lock.dart
+++ b/Engine/Extras/VirtualProduction/EpicStageApp/lib/widgets/elements/control_lock.dart
@@ -1,12 +1,11 @@
 // Copyright Epic Games, Inc. All Rights Reserved.
 
+import 'package:epic_common/theme.dart';
+import 'package:epic_common/widgets.dart';
 import 'package:flutter/material.dart';
 import 'package:flutter/services.dart';
 import 'package:flutter_gen/gen_l10n/app_localizations.dart';
 
-import '../../utilities/unreal_colors.dart';
-import 'epic_icon_button.dart';
-
 /// Widget housing for locking or unlocking the screen, either ignoring touch input on allowing them.
 class ControlLock extends StatefulWidget {
   const ControlLock({Key? key}) : super(key: key);
@@ -128,7 +127,9 @@ class _ControlLockButton extends StatelessWidget {
           : EpicIconButton(
               bIsToggledOn: bIsLocked,
               activeBackgroundColor: UnrealColors.highlightRed,
-              iconPath: bIsLocked ? 'assets/images/icons/lock_locked.svg' : 'assets/images/icons/lock_unlocked.svg',
+              iconPath: bIsLocked
+                  ? 'packages/epic_common/assets/icons/lock_locked.svg'
+                  : 'packages/epic_common/assets/icons/lock_unlocked.svg',
               onLongPressed: onLongPressed,
               tooltipMessage:
                   bIsLocked ? AppLocalizations.of(context)!.unlockControls : AppLocalizations.of(context)!.lockControls,
diff --git a/Engine/Extras/VirtualProduction/EpicStageApp/lib/widgets/elements/delta_slider.dart b/Engine/Extras/VirtualProduction/EpicStageApp/lib/widgets/elements/delta_slider.dart
--- a/Engine/Extras/VirtualProduction/EpicStageApp/lib/widgets/elements/delta_slider.dart
+++ b/Engine/Extras/VirtualProduction/EpicStageApp/lib/widgets/elements/delta_slider.dart
@@ -2,6 +2,9 @@
 
 import 'dart:math' as math;
 
+import 'package:epic_common/preferences.dart';
+import 'package:epic_common/theme.dart';
+import 'package:epic_common/widgets.dart';
 import 'package:flutter/material.dart';
 import 'package:flutter_gen/gen_l10n/app_localizations.dart';
 import 'package:provider/provider.dart';
@@ -9,10 +12,7 @@ import 'package:provider/provider.dart';
 import '../../models/property_modify_operations.dart';
 import '../../models/settings/delta_widget_settings.dart';
 import '../../utilities/math_utilities.dart';
-import '../../utilities/transient_preference.dart';
-import '../../utilities/unreal_colors.dart';
 import 'delta_widget_base.dart';
-import 'modal.dart';
 import 'unreal_widget_base.dart';
 
 const double _dotSize = 17.0;
diff --git a/Engine/Extras/VirtualProduction/EpicStageApp/lib/widgets/elements/delta_widget_base.dart b/Engine/Extras/VirtualProduction/EpicStageApp/lib/widgets/elements/delta_widget_base.dart
--- a/Engine/Extras/VirtualProduction/EpicStageApp/lib/widgets/elements/delta_widget_base.dart
+++ b/Engine/Extras/VirtualProduction/EpicStageApp/lib/widgets/elements/delta_widget_base.dart
@@ -1,11 +1,11 @@
 // Copyright Epic Games, Inc. All Rights Reserved.
 
+import 'package:epic_common/widgets.dart';
 import 'package:flutter/material.dart';
 import 'package:flutter_gen/gen_l10n/app_localizations.dart';
 import 'package:provider/provider.dart';
 
 import '../../models/settings/delta_widget_settings.dart';
-import 'epic_icon_button.dart';
 
 class DeltaWidgetConstants {
   /// Size of icon buttons used for delta widgets.
@@ -43,7 +43,7 @@ class ResetValueButton extends StatelessWidget {
   Widget build(BuildContext context) {
     return EpicIconButton(
       onPressed: onPressed,
-      iconPath: 'assets/images/icons/reset.svg',
+      iconPath: 'packages/epic_common/assets/icons/reset.svg',
       tooltipMessage: AppLocalizations.of(context)!.propertyResetButtonTooltip,
       iconSize: 24,
       buttonSize: const Size(32, 32),
diff --git a/Engine/Extras/VirtualProduction/EpicStageApp/lib/widgets/elements/dropdown_button.dart b/Engine/Extras/VirtualProduction/EpicStageApp/lib/widgets/elements/dropdown_button.dart
--- a/Engine/Extras/VirtualProduction/EpicStageApp/lib/widgets/elements/dropdown_button.dart
+++ b/Engine/Extras/VirtualProduction/EpicStageApp/lib/widgets/elements/dropdown_button.dart
@@ -1,303 +1,15 @@
 // Copyright Epic Games, Inc. All Rights Reserved.
 
+import 'package:epic_common/preferences.dart';
+import 'package:epic_common/widgets.dart';
 import 'package:flutter/material.dart';
 import 'package:provider/provider.dart';
 
-import '../../models/navigator_keys.dart';
 import '../../models/property_modify_operations.dart';
 import '../../models/settings/delta_widget_settings.dart';
-import '../../utilities/transient_preference.dart';
-import '../../utilities/unreal_colors.dart';
-import 'asset_icon.dart';
 import 'delta_widget_base.dart';
-import 'dropdown_list_menu.dart';
-import 'list_menu.dart';
 import 'unreal_widget_base.dart';
 
-/// Possible states for a [ModalTextDropdownButton].
-enum ModalDropdownButtonState {
-  /// The modal is closed.
-  closed,
-
-  /// The modal is open and displaying under the button.
-  openOnTop,
-
-  /// The modal is open and displaying above the button.
-  openOnBottom,
-}
-
-/// A button that uses a builder function to open a dropdown list modal when pressed.
-class ModalDropdownButton<ItemType> extends StatefulWidget {
-  const ModalDropdownButton({
-    Key? key,
-    required this.buttonBuilder,
-    required this.menuBuilder,
-    this.openRectExtraSpace = EdgeInsets.zero,
-    this.bStretchDropdown = false,
-    this.tooltipMessage,
-    this.bDisabled = false,
-  }) : super(key: key);
-
-  /// Function that builds the button based on the dropdown's [state].
-  final Widget Function(BuildContext context, ModalDropdownButtonState state) buttonBuilder;
-
-  /// Function that builds the dropdown menu.
-  /// Receives an [originTabBuilder] function which will build the menu's origin tab based on whether the tab [isOnTop].
-  final Widget Function(BuildContext context, Widget Function(BuildContext, bool isOnTop) originTabBuilder) menuBuilder;
-
-  /// How much space to add around the original button's bounding area when building its open version.
-  final EdgeInsets openRectExtraSpace;
-
-  /// Whether to stretch the dropdown to match the size of the button.
-  final bool bStretchDropdown;
-
-  /// An optional tooltip message to show when the button is hovered or long pressed.
-  final String? tooltipMessage;
-
-  /// whether the button is disabled or not.
-  final bool? bDisabled;
-
-  @override
-  State<ModalDropdownButton<ItemType>> createState() => _ModalDropdownButtonState<ItemType>();
-}
-
-class _ModalDropdownButtonState<ItemType> extends State<ModalDropdownButton<ItemType>> {
-  /// Key used to identify the widget from which the dropdown originates.
-  final GlobalKey _dropdownOriginator = GlobalKey();
-
-  @override
-  Widget build(BuildContext context) {
-    final button = Semantics(
-      button: true,
-      child: MouseRegion(
-        cursor: MaterialStateMouseCursor.clickable,
-        child: GestureDetector(
-          key: _dropdownOriginator,
-          onTap: widget.bDisabled == true ? null : _onTap,
-          child: widget.buttonBuilder(context, ModalDropdownButtonState.closed),
-        ),
-      ),
-    );
-
-    return (widget.tooltipMessage != null)
-        ? Tooltip(
-            child: button,
-            message: widget.tooltipMessage,
-          )
-        : button;
-  }
-
-  /// Called when this is tapped.
-  void _onTap() {
-    final RenderObject? widgetRenderObject = _dropdownOriginator.currentContext?.findRenderObject();
-    final RenderBox? widgetBox = (widgetRenderObject != null) ? (widgetRenderObject as RenderBox) : null;
-
-    if (widgetBox == null) {
-      throw Exception('Failed to find widget render box for $this');
-    }
-
-    final RenderObject rootRenderObject = rootNavigatorKey.currentContext!.findRenderObject()!;
-
-    final Size rectSize = Size(
-      widgetBox.size.width + widget.openRectExtraSpace.horizontal,
-      widgetBox.size.height + widget.openRectExtraSpace.vertical,
-    );
-
-    final Offset globalPos = widgetBox.localToGlobal(
-      Offset(-widget.openRectExtraSpace.left, -widget.openRectExtraSpace.bottom),
-      ancestor: rootRenderObject,
-    );
-
-    final Rect widgetRect = Rect.fromLTWH(globalPos.dx, globalPos.dy, rectSize.width, rectSize.height);
-
-    DropDownListMenu.showAtRect(
-      pivotRect: widgetRect,
-      bStretch: widget.bStretchDropdown,
-      builder: (context) => widget.menuBuilder(
-        context,
-        (context, bIsOnTop) => ConstrainedBox(
-          constraints: BoxConstraints(
-            maxWidth: widgetRect.width,
-            maxHeight: widgetRect.height,
-          ),
-          child: widget.buttonBuilder(
-            context,
-            bIsOnTop ? ModalDropdownButtonState.openOnTop : ModalDropdownButtonState.openOnBottom,
-          ),
-        ),
-      ),
-    );
-  }
-}
-
-/// A button that opens an automatically built dropdown list modal containing text-only options when pressed.
-class ModalTextDropdownButton<ItemType> extends StatelessWidget {
-  const ModalTextDropdownButton({
-    Key? key,
-    required this.builder,
-    required this.items,
-    required this.onChanged,
-    this.value,
-    this.makeItemName,
-    this.openRectExtraSpace = EdgeInsets.zero,
-    this.bStretchDropdown = false,
-  }) : super(key: key);
-
-  /// Function that builds the button based on the dropdown's state.
-  final Widget Function(BuildContext context, ModalDropdownButtonState state) builder;
-
-  /// The currently selected value.
-  final ItemType? value;
-
-  /// The list of selectable items in the order in which they should be displayed.
-  final List<ItemType> items;
-
-  /// Function called when a new value is selected.
-  final void Function(ItemType newValue) onChanged;
-
-  /// Optional function to convert items to a display string. If not provided, the items will be converted implicitly.
-  final String Function(ItemType item)? makeItemName;
-
-  /// How much space to add around the original button's bounding area when building its open version.
-  final EdgeInsets openRectExtraSpace;
-
-  /// Whether to stretch the dropdown to match the size of the button.
-  final bool bStretchDropdown;
-
-  @override
-  Widget build(BuildContext context) {
-    return ModalDropdownButton(
-      openRectExtraSpace: openRectExtraSpace,
-      bStretchDropdown: bStretchDropdown,
-      buttonBuilder: builder,
-      menuBuilder: (context, originTabBuilder) => DropDownListMenu(
-        originTabBuilder: originTabBuilder,
-        children: [
-          for (ItemType item in items)
-            ListMenuSimpleItem(
-              title: _makeItemName(item),
-              onTap: () => _onValueSelected(context, item),
-            ),
-        ],
-      ),
-    );
-  }
-
-  /// Convert an item to a display name, using the user-provided function if possible.
-  String _makeItemName(ItemType item) {
-    if (makeItemName != null) {
-      return makeItemName!(item);
-    }
-
-    return item.toString();
-  }
-
-  /// Called when an item is selected.
-  void _onValueSelected(BuildContext context, ItemType value) {
-    final navigator = Navigator.of(context, rootNavigator: true);
-    navigator.pop();
-
-    onChanged(value);
-  }
-}
-
-/// A dropdown button styled as a large selector containing text.
-class DropdownSelector<ItemType> extends StatelessWidget {
-  const DropdownSelector({
-    Key? key,
-    required this.value,
-    required this.items,
-    required this.onChanged,
-    this.hint = '',
-    this.makeItemName,
-    this.bIsEnabled = true,
-    this.borderRadius = const Radius.circular(4),
-  }) : super(key: key);
-
-  /// The currently selected value.
-  final ItemType? value;
-
-  /// The list of selectable items in the order in which they should be displayed.
-  final List<ItemType> items;
-
-  /// Function called when a new value is selected.
-  final void Function(ItemType newValue) onChanged;
-
-  /// Optional function to convert items to a display string. If not provided, the items will be converted implicitly.
-  final String Function(ItemType item)? makeItemName;
-
-  /// If false, disable the controls and grey this out.
-  final bool bIsEnabled;
-
-  /// Text to show when [value] is null.
-  final String hint;
-
-  /// The border radius of the button.
-  final Radius borderRadius;
-
-  /// The display name to show inside the dropdown button.
-  String get _displayName {
-    if (value == null) {
-      return hint;
-    }
-
-    if (makeItemName != null) {
-      return makeItemName!(value!);
-    }
-
-    return value.toString();
-  }
-
-  @override
-  Widget build(BuildContext context) {
-    final Color textColor = bIsEnabled ? UnrealColors.white : UnrealColors.gray56;
-
-    final builder = (BuildContext context, ModalDropdownButtonState state) => Container(
-          height: 36,
-          decoration: BoxDecoration(
-            color: bIsEnabled ? Theme.of(context).colorScheme.background : UnrealColors.gray22,
-            borderRadius: BorderRadius.vertical(
-              top: borderRadius,
-              bottom: state == ModalDropdownButtonState.closed ? borderRadius : Radius.zero,
-            ),
-          ),
-          padding: EdgeInsets.symmetric(horizontal: 12),
-          child: Row(
-            mainAxisAlignment: MainAxisAlignment.spaceBetween,
-            crossAxisAlignment: CrossAxisAlignment.center,
-            children: [
-              Text(
-                _displayName,
-                style: Theme.of(context).textTheme.bodyMedium!.copyWith(
-                      color: textColor,
-                      letterSpacing: 0.5,
-                    ),
-              ),
-              SizedBox(width: 8),
-              AssetIcon(
-                path: 'assets/images/icons/chevron_down.svg',
-                size: 16,
-                color: textColor,
-              ),
-            ],
-          ),
-        );
-
-    if (bIsEnabled) {
-      return ModalTextDropdownButton(
-        bStretchDropdown: true,
-        value: value,
-        items: items,
-        onChanged: onChanged,
-        makeItemName: makeItemName,
-        builder: builder,
-      );
-    } else {
-      return builder(context, ModalDropdownButtonState.closed);
-    }
-  }
-}
-
 /// A dropdown that controls one or more enum properties in Unreal Engine.
 class UnrealDropdownSelector extends UnrealWidget {
   const UnrealDropdownSelector({
diff --git a/Engine/Extras/VirtualProduction/EpicStageApp/lib/widgets/elements/dropdown_list_menu.dart b/Engine/Extras/VirtualProduction/EpicStageApp/lib/widgets/elements/dropdown_list_menu.dart
deleted file mode 100644
diff --git a/Engine/Extras/VirtualProduction/EpicStageApp/lib/widgets/elements/dropdown_text.dart b/Engine/Extras/VirtualProduction/EpicStageApp/lib/widgets/elements/dropdown_text.dart
--- a/Engine/Extras/VirtualProduction/EpicStageApp/lib/widgets/elements/dropdown_text.dart
+++ b/Engine/Extras/VirtualProduction/EpicStageApp/lib/widgets/elements/dropdown_text.dart
@@ -1,25 +1,15 @@
 // Copyright Epic Games, Inc. All Rights Reserved.
 
+import 'package:epic_common/preferences.dart';
+import 'package:epic_common/widgets.dart';
 import 'package:flutter/material.dart';
 import 'package:flutter_gen/gen_l10n/app_localizations.dart';
 import 'package:provider/provider.dart';
 
 import '../../models/property_modify_operations.dart';
 import '../../models/settings/delta_widget_settings.dart';
-import '../../utilities/transient_preference.dart';
-import '../../utilities/unreal_colors.dart';
-import 'asset_icon.dart';
-import 'dropdown_button.dart';
 import 'unreal_widget_base.dart';
 
-/// Padding used inside the origin tab when the dropdown is open.
-const EdgeInsets _originTabPadding = EdgeInsets.only(
-  left: 8,
-  right: 4,
-  top: 12,
-  bottom: 12,
-);
-
 /// Text which acts as a dropdown, but otherwise remains inline.
 /// Controls an enum value in Unreal Engine.
 class UnrealDropdownText extends UnrealWidget {
@@ -46,7 +36,7 @@ class _UnrealDropdownTextState extends State<UnrealDropdownText>
 
     return TransientPreferenceBuilder(
       preference: Provider.of<DeltaWidgetSettings>(context, listen: false).bIsInResetMode,
-      builder: (context, final bool bIsInResetMode) => DrivenDropdownText<SingleSharedValue<String>>(
+      builder: (context, final bool bIsInResetMode) => DropdownText<SingleSharedValue<String>>(
         value: sharedValue,
         items: propertyEnumValues
             .map((enumValue) => SingleSharedValue(
@@ -87,178 +77,3 @@ class _UnrealDropdownTextState extends State<UnrealDropdownText>
     endTransaction();
   }
 }
-
-/// Text which acts as a dropdown, but otherwise remains inline.
-class DrivenDropdownText<ItemType> extends StatelessWidget {
-  const DrivenDropdownText({
-    Key? key,
-    required this.value,
-    required this.items,
-    required this.onChanged,
-    this.makeItemName,
-    this.bIsEnabled = true,
-  }) : super(key: key);
-
-  /// The currently selected value.
-  final ItemType value;
-
-  /// The list of selectable items in the order in which they should be displayed.
-  final List<ItemType> items;
-
-  /// Function called when a new value is selected.
-  final void Function(ItemType newValue) onChanged;
-
-  /// Optional function to convert items to a display string. If not provided, the items will be converted implicitly.
-  final String? Function(ItemType item)? makeItemName;
-
-  /// If false, grey out the widget and disable its controls.
-  final bool bIsEnabled;
-
-  @override
-  Widget build(BuildContext context) {
-    // Align the row based on text directionality, respecting the text direction.
-    late final MainAxisAlignment alignment;
-    switch (DefaultTextStyle.of(context).textAlign) {
-      case TextAlign.start:
-        alignment = MainAxisAlignment.start;
-        break;
-
-      case TextAlign.end:
-        alignment = MainAxisAlignment.end;
-        break;
-
-      case TextAlign.center:
-        alignment = MainAxisAlignment.center;
-        break;
-
-      case TextAlign.left:
-        alignment = (Directionality.of(context) == TextDirection.ltr) ? MainAxisAlignment.start : MainAxisAlignment.end;
-        break;
-
-      case TextAlign.right:
-        alignment = (Directionality.of(context) == TextDirection.ltr) ? MainAxisAlignment.end : MainAxisAlignment.start;
-        break;
-
-      default:
-        alignment = MainAxisAlignment.start;
-    }
-
-    if (!bIsEnabled) {
-      return DefaultTextStyle(
-        style: DefaultTextStyle.of(context).style.copyWith(
-              color: DefaultTextStyle.of(context).style.color!.withOpacity(0.4),
-            ),
-        child: _DrivenDropdownTextInner(
-          text: _makeItemName(value),
-          alignment: alignment,
-          bHideChevron: true,
-        ),
-      );
-    }
-
-    return ModalTextDropdownButton(
-      value: value,
-      items: items,
-      onChanged: onChanged,
-      makeItemName: _makeItemName,
-      openRectExtraSpace: _originTabPadding,
-      builder: (BuildContext innerContext, ModalDropdownButtonState state) {
-        final String name = _makeItemName(value);
-
-        if (state == ModalDropdownButtonState.closed) {
-          return _DrivenDropdownTextInner(
-            text: name,
-            alignment: alignment,
-          );
-        }
-
-        final bool bIsOnTop = state == ModalDropdownButtonState.openOnTop;
-
-        return Container(
-          decoration: BoxDecoration(
-            color: UnrealColors.gray22,
-            borderRadius: BorderRadius.vertical(
-              top: bIsOnTop ? Radius.circular(8) : Radius.zero,
-              bottom: !bIsOnTop ? Radius.circular(8) : Radius.zero,
-            ),
-          ),
-          padding: _originTabPadding,
-          child: DefaultTextStyle(
-            style: DefaultTextStyle.of(context).style,
-            child: _DrivenDropdownTextInner(
-              text: name,
-              alignment: MainAxisAlignment.start,
-              color: UnrealColors.white,
-            ),
-          ),
-        );
-      },
-    );
-  }
-
-  /// Convert an item to a display name, using the user-provided function if possible.
-  String _makeItemName(ItemType item) {
-    if (item == null) {
-      return '...';
-    }
-
-    if (makeItemName != null) {
-      return makeItemName!(item) ?? '...';
-    }
-
-    return item.toString();
-  }
-}
-
-/// The inner contents of a [DrivenDropdownText], showing the text and the chevron icon.
-class _DrivenDropdownTextInner extends StatelessWidget {
-  const _DrivenDropdownTextInner({
-    Key? key,
-    required this.text,
-    required this.alignment,
-    this.color,
-    this.bHideChevron = false,
-  }) : super(key: key);
-
-  final String text;
-  final MainAxisAlignment alignment;
-  final Color? color;
-  final bool bHideChevron;
-
-  @override
-  Widget build(BuildContext context) {
-    TextStyle textStyle = DefaultTextStyle.of(context).style;
-    if (color != null) {
-      textStyle = textStyle.copyWith(color: color);
-    }
-
-    return Row(
-      mainAxisSize: MainAxisSize.min,
-      mainAxisAlignment: alignment,
-      children: [
-        Flexible(
-          child: Text(
-            text,
-            overflow: TextOverflow.ellipsis,
-            softWrap: false,
-            style: textStyle,
-          ),
-        ),
-        if (!bHideChevron)
-          SizedOverflowBox(
-            size: const Size(24, 15),
-            alignment: Alignment.center,
-            child: SizedBox(
-              width: 24,
-              child: AssetIcon(
-                path: 'assets/images/icons/chevron_down.svg',
-                size: 24,
-                color: textStyle.color,
-                fit: BoxFit.cover,
-              ),
-            ),
-          ),
-      ],
-    );
-  }
-}
diff --git a/Engine/Extras/VirtualProduction/EpicStageApp/lib/widgets/elements/empty_placeholder.dart b/Engine/Extras/VirtualProduction/EpicStageApp/lib/widgets/elements/empty_placeholder.dart
deleted file mode 100644
diff --git a/Engine/Extras/VirtualProduction/EpicStageApp/lib/widgets/elements/epic_checkbox.dart b/Engine/Extras/VirtualProduction/EpicStageApp/lib/widgets/elements/epic_checkbox.dart
deleted file mode 100644
diff --git a/Engine/Extras/VirtualProduction/EpicStageApp/lib/widgets/elements/epic_icon_button.dart b/Engine/Extras/VirtualProduction/EpicStageApp/lib/widgets/elements/epic_icon_button.dart
deleted file mode 100644
diff --git a/Engine/Extras/VirtualProduction/EpicStageApp/lib/widgets/elements/epic_list_view.dart b/Engine/Extras/VirtualProduction/EpicStageApp/lib/widgets/elements/epic_list_view.dart
deleted file mode 100644
diff --git a/Engine/Extras/VirtualProduction/EpicStageApp/lib/widgets/elements/epic_scroll_view.dart b/Engine/Extras/VirtualProduction/EpicStageApp/lib/widgets/elements/epic_scroll_view.dart
deleted file mode 100644
diff --git a/Engine/Extras/VirtualProduction/EpicStageApp/lib/widgets/elements/floating_map_preview.dart b/Engine/Extras/VirtualProduction/EpicStageApp/lib/widgets/elements/floating_map_preview.dart
--- a/Engine/Extras/VirtualProduction/EpicStageApp/lib/widgets/elements/floating_map_preview.dart
+++ b/Engine/Extras/VirtualProduction/EpicStageApp/lib/widgets/elements/floating_map_preview.dart
@@ -2,13 +2,13 @@
 
 import 'dart:math' as math;
 
+import 'package:epic_common/preferences.dart';
+import 'package:epic_common/widgets.dart';
 import 'package:flutter/material.dart';
 import 'package:flutter/physics.dart';
 import 'package:flutter/scheduler.dart';
 
 import '../../models/settings/floating_map_settings.dart';
-import '../../utilities/preferences_bundle.dart';
-import 'asset_icon.dart';
 import 'lightcard_map.dart';
 
 /// Sides of the screen on which the preview can be docked.
@@ -670,14 +670,16 @@ class _DockTab extends StatelessWidget {
       ),
       child: Column(mainAxisAlignment: MainAxisAlignment.center, children: [
         AssetIcon(
-          path: 'assets/images/icons/map.svg',
+          path: 'packages/epic_common/assets/icons/map.svg',
           size: 24,
         ),
         SizedBox(
           height: 6,
         ),
         AssetIcon(
-          path: bIsFacingLeft ? 'assets/images/icons/chevron_left.svg' : 'assets/images/icons/chevron_right.svg',
+          path: bIsFacingLeft
+              ? 'packages/epic_common/assets/icons/chevron_left.svg'
+              : 'packages/epic_common/assets/icons/chevron_right.svg',
           size: 24,
         ),
       ]),
diff --git a/Engine/Extras/VirtualProduction/EpicStageApp/lib/widgets/elements/layout/card.dart b/Engine/Extras/VirtualProduction/EpicStageApp/lib/widgets/elements/layout/card.dart
deleted file mode 100644
diff --git a/Engine/Extras/VirtualProduction/EpicStageApp/lib/widgets/elements/layout/lazy_indexed_stack.dart b/Engine/Extras/VirtualProduction/EpicStageApp/lib/widgets/elements/layout/lazy_indexed_stack.dart
deleted file mode 100644
diff --git a/Engine/Extras/VirtualProduction/EpicStageApp/lib/widgets/elements/layout/lazy_tab_view.dart b/Engine/Extras/VirtualProduction/EpicStageApp/lib/widgets/elements/layout/lazy_tab_view.dart
deleted file mode 100644
diff --git a/Engine/Extras/VirtualProduction/EpicStageApp/lib/widgets/elements/layout/measurable.dart b/Engine/Extras/VirtualProduction/EpicStageApp/lib/widgets/elements/layout/measurable.dart
deleted file mode 100644
diff --git a/Engine/Extras/VirtualProduction/EpicStageApp/lib/widgets/elements/lightcard_map.dart b/Engine/Extras/VirtualProduction/EpicStageApp/lib/widgets/elements/lightcard_map.dart
--- a/Engine/Extras/VirtualProduction/EpicStageApp/lib/widgets/elements/lightcard_map.dart
+++ b/Engine/Extras/VirtualProduction/EpicStageApp/lib/widgets/elements/lightcard_map.dart
@@ -4,6 +4,9 @@ import 'dart:async';
 import 'dart:math' as math;
 import 'dart:typed_data';
 
+import 'package:epic_common/preferences.dart';
+import 'package:epic_common/theme.dart';
+import 'package:epic_common/widgets.dart';
 import 'package:flutter/gestures.dart';
 import 'package:flutter/material.dart';
 import 'package:flutter_gen/gen_l10n/app_localizations.dart';
@@ -23,12 +26,7 @@ import '../../models/unreal_property_manager.dart';
 import '../../models/unreal_types.dart';
 import '../../utilities/constants.dart';
 import '../../utilities/guarded_refresh_state.dart';
-import '../../utilities/transient_preference.dart';
-import '../../utilities/unreal_colors.dart';
-import 'asset_icon.dart';
 import 'dropdown_button.dart';
-import 'dropdown_list_menu.dart';
-import 'epic_icon_button.dart';
 import 'place_actor_menu.dart';
 import 'spinner_overlay.dart';
 import 'transform_gesture_detector.dart';
@@ -976,6 +974,7 @@ class StageMapState extends State<StageMap> with PreviewRenderConsumer, GuardedR
   /// Called when the user performs a long press gesture on the map.
   void _onLongPressStart(LongPressStartDetails details) {
     DropDownListMenu.showAtPosition(
+      context,
       pivotPosition: details.globalPosition,
       builder: (context) => PlaceActorDropDownMenu(
         actorMapPosition: _globalToMapPosition(details.globalPosition),
@@ -1914,13 +1913,13 @@ class _StageMapControlModeToggle extends StatelessWidget {
         preference: stageMapSettings.bIsInTransformMode,
         builder: (_, final bool bIsInTransformMode) => Row(children: [
           _StageMapControlModeToggleButton(
-            iconPath: 'assets/images/icons/world_flat.svg',
+            iconPath: 'packages/epic_common/assets/icons/world_flat.svg',
             tooltipMessage: AppLocalizations.of(context)!.stageMapControlModeMap,
             bIsActive: !bIsInTransformMode,
             onPressed: () => stageMapSettings.bIsInTransformMode.setValue(false),
           ),
           _StageMapControlModeToggleButton(
-            iconPath: 'assets/images/icons/object.svg',
+            iconPath: 'packages/epic_common/assets/icons/object.svg',
             tooltipMessage: AppLocalizations.of(context)!.stageMapControlModeObject,
             bIsActive: bIsInTransformMode,
             onPressed: () => stageMapSettings.bIsInTransformMode.setValue(true),
diff --git a/Engine/Extras/VirtualProduction/EpicStageApp/lib/widgets/elements/list_menu.dart b/Engine/Extras/VirtualProduction/EpicStageApp/lib/widgets/elements/list_menu.dart
deleted file mode 100644
diff --git a/Engine/Extras/VirtualProduction/EpicStageApp/lib/widgets/elements/modal.dart b/Engine/Extras/VirtualProduction/EpicStageApp/lib/widgets/elements/modal.dart
deleted file mode 100644
diff --git a/Engine/Extras/VirtualProduction/EpicStageApp/lib/widgets/elements/ndisplay_config_name_display.dart b/Engine/Extras/VirtualProduction/EpicStageApp/lib/widgets/elements/ndisplay_config_name_display.dart
--- a/Engine/Extras/VirtualProduction/EpicStageApp/lib/widgets/elements/ndisplay_config_name_display.dart
+++ b/Engine/Extras/VirtualProduction/EpicStageApp/lib/widgets/elements/ndisplay_config_name_display.dart
@@ -1,5 +1,6 @@
 // Copyright Epic Games, Inc. All Rights Reserved.
 
+import 'package:epic_common/widgets.dart';
 import 'package:flutter/material.dart';
 import 'package:provider/provider.dart';
 
@@ -7,7 +8,6 @@ import '../../models/settings/selected_actor_settings.dart';
 import '../../models/unreal_actor_manager.dart';
 import '../../models/unreal_types.dart';
 import '../../utilities/constants.dart';
-import 'asset_icon.dart';
 
 class NDisplayConfigNameDisplay extends StatefulWidget {
   const NDisplayConfigNameDisplay({Key? key}) : super(key: key);
@@ -45,7 +45,7 @@ class _NDisplayConfigNameDisplayState extends State<NDisplayConfigNameDisplay> {
           crossAxisAlignment: CrossAxisAlignment.center,
           children: [
             AssetIcon(
-              path: 'assets/images/icons/ndisplay.svg',
+              path: 'packages/epic_common/assets/icons/ndisplay.svg',
               size: 20,
             ),
             const SizedBox(width: 8),
diff --git a/Engine/Extras/VirtualProduction/EpicStageApp/lib/widgets/elements/parsed_rich_text.dart b/Engine/Extras/VirtualProduction/EpicStageApp/lib/widgets/elements/parsed_rich_text.dart
deleted file mode 100644
diff --git a/Engine/Extras/VirtualProduction/EpicStageApp/lib/widgets/elements/place_actor_menu.dart b/Engine/Extras/VirtualProduction/EpicStageApp/lib/widgets/elements/place_actor_menu.dart
--- a/Engine/Extras/VirtualProduction/EpicStageApp/lib/widgets/elements/place_actor_menu.dart
+++ b/Engine/Extras/VirtualProduction/EpicStageApp/lib/widgets/elements/place_actor_menu.dart
@@ -1,5 +1,7 @@
 // Copyright Epic Games, Inc. All Rights Reserved.
 
+import 'package:epic_common/theme.dart';
+import 'package:epic_common/widgets.dart';
 import 'package:flutter/material.dart';
 import 'package:flutter_gen/gen_l10n/app_localizations.dart';
 import 'package:provider/provider.dart';
@@ -8,11 +10,7 @@ import 'package:streaming_shared_preferences/streaming_shared_preferences.dart';
 import '../../models/settings/recent_actor_settings.dart';
 import '../../models/unreal_actor_creator.dart';
 import '../../models/unreal_types.dart';
-import '../../utilities/unreal_colors.dart';
 import 'dropdown_button.dart';
-import 'dropdown_list_menu.dart';
-import 'epic_icon_button.dart';
-import 'list_menu.dart';
 import 'template_menu.dart';
 
 /// Button that opens the actor placement menu.
@@ -67,7 +65,7 @@ class _PlaceActorButtonContents extends StatelessWidget {
         widthFactor: 1,
         child: EpicWideButton(
           text: AppLocalizations.of(context)!.addActorButton,
-          iconPath: 'assets/images/icons/plus.svg',
+          iconPath: 'packages/epic_common/assets/icons/plus.svg',
           color: bIsActive ? UnrealColors.highlightGreen : Theme.of(context).colorScheme.surfaceTint,
           iconColor: bIsActive ? null : UnrealColors.highlightGreen,
           border: Border.all(
@@ -110,7 +108,7 @@ class PlaceActorDropDownMenu extends StatelessWidget {
       children: [
         ListMenuSimpleItem(
           title: AppLocalizations.of(context)!.actorNameLightCard,
-          iconPath: 'assets/images/icons/light_card.svg',
+          iconPath: 'packages/epic_common/assets/icons/light_card.svg',
           onTap: () {
             actorCreator.createLightcard(actorCreationSettings);
             _closeMenu(context);
@@ -118,7 +116,7 @@ class PlaceActorDropDownMenu extends StatelessWidget {
         ),
         ListMenuSimpleItem(
           title: AppLocalizations.of(context)!.actorNameFlag,
-          iconPath: 'assets/images/icons/light_card_flag.svg',
+          iconPath: 'packages/epic_common/assets/icons/light_card_flag.svg',
           onTap: () {
             actorCreator.createFlag(actorCreationSettings);
             _closeMenu(context);
@@ -126,7 +124,7 @@ class PlaceActorDropDownMenu extends StatelessWidget {
         ),
         ListMenuSimpleItem(
           title: AppLocalizations.of(context)!.actorNameChromakeyCard,
-          iconPath: 'assets/images/icons/chromakey_card.svg',
+          iconPath: 'packages/epic_common/assets/icons/chromakey_card.svg',
           onTap: () {
             actorCreator.createChromakeyCard(actorCreationSettings);
             _closeMenu(context);
@@ -134,7 +132,7 @@ class PlaceActorDropDownMenu extends StatelessWidget {
         ),
         ListMenuSimpleItem(
           title: AppLocalizations.of(context)!.actorNameColorCorrectWindow,
-          iconPath: 'assets/images/icons/color_correct_window.svg',
+          iconPath: 'packages/epic_common/assets/icons/color_correct_window.svg',
           onTap: () {
             actorCreator.createColorCorrectionWindow(actorCreationSettings);
             _closeMenu(context);
@@ -142,7 +140,7 @@ class PlaceActorDropDownMenu extends StatelessWidget {
         ),
         ListMenuSimpleItem(
           title: AppLocalizations.of(context)!.placeActorMenuAllTemplates,
-          iconPath: 'assets/images/icons/template.svg',
+          iconPath: 'packages/epic_common/assets/icons/template.svg',
           bShowArrow: true,
           onTap: () {
             _closeMenu(context);
@@ -190,21 +188,21 @@ class _RecentlyPlacedMenuItem extends StatelessWidget {
   /// The path of the icon to display for this item.
   String? get _iconPath {
     if (actorData.templatePath != null) {
-      return 'assets/images/icons/template.svg';
+      return 'packages/epic_common/assets/icons/template.svg';
     }
 
     switch (actorData.type) {
       case RecentlyPlacedActorType.lightCard:
-        return 'assets/images/icons/light_card.svg';
+        return 'packages/epic_common/assets/icons/light_card.svg';
 
       case RecentlyPlacedActorType.flag:
-        return 'assets/images/icons/light_card_flag.svg';
+        return 'packages/epic_common/assets/icons/light_card_flag.svg';
 
       case RecentlyPlacedActorType.chromakeyCard:
-        return 'assets/images/icons/chromakey_card.svg';
+        return 'packages/epic_common/assets/icons/chromakey_card.svg';
 
       case RecentlyPlacedActorType.colorCorrectWindow:
-        return 'assets/images/icons/color_correct_window.svg';
+        return 'packages/epic_common/assets/icons/color_correct_window.svg';
 
       default:
         return null;
diff --git a/Engine/Extras/VirtualProduction/EpicStageApp/lib/widgets/elements/reset_mode_button.dart b/Engine/Extras/VirtualProduction/EpicStageApp/lib/widgets/elements/reset_mode_button.dart
--- a/Engine/Extras/VirtualProduction/EpicStageApp/lib/widgets/elements/reset_mode_button.dart
+++ b/Engine/Extras/VirtualProduction/EpicStageApp/lib/widgets/elements/reset_mode_button.dart
@@ -1,13 +1,13 @@
 // Copyright Epic Games, Inc. All Rights Reserved.
 
+import 'package:epic_common/preferences.dart';
+import 'package:epic_common/theme.dart';
+import 'package:epic_common/widgets.dart';
 import 'package:flutter/material.dart';
 import 'package:flutter_gen/gen_l10n/app_localizations.dart';
 import 'package:provider/provider.dart';
 
 import '../../models/settings/delta_widget_settings.dart';
-import '../../utilities/constants.dart';
-import '../../utilities/transient_preference.dart';
-import 'epic_icon_button.dart';
 
 /// Button to toggle reset mode on and off.
 class ResetModeButton extends StatelessWidget {
@@ -18,11 +18,11 @@ class ResetModeButton extends StatelessWidget {
     final deltaSettings = Provider.of<DeltaWidgetSettings>(context);
 
     return Padding(
-      padding: EdgeInsets.only(right: cardMargin),
+      padding: EdgeInsets.only(right: UnrealTheme.cardMargin),
       child: TransientPreferenceBuilder(
         preference: Provider.of<DeltaWidgetSettings>(context, listen: false).bIsInResetMode,
         builder: (context, final bool bIsInResetMode) => EpicIconButton(
-          iconPath: 'assets/images/icons/reset_small.svg',
+          iconPath: 'packages/epic_common/assets/icons/reset_small.svg',
           iconSize: 17,
           backgroundSize: const Size(33, 33),
           tooltipMessage: AppLocalizations.of(context)!.resetModeToggleButtonTooltip,
diff --git a/Engine/Extras/VirtualProduction/EpicStageApp/lib/widgets/elements/selector_bar.dart b/Engine/Extras/VirtualProduction/EpicStageApp/lib/widgets/elements/selector_bar.dart
deleted file mode 100644
diff --git a/Engine/Extras/VirtualProduction/EpicStageApp/lib/widgets/elements/sensitivity_slider.dart b/Engine/Extras/VirtualProduction/EpicStageApp/lib/widgets/elements/sensitivity_slider.dart
--- a/Engine/Extras/VirtualProduction/EpicStageApp/lib/widgets/elements/sensitivity_slider.dart
+++ b/Engine/Extras/VirtualProduction/EpicStageApp/lib/widgets/elements/sensitivity_slider.dart
@@ -1,17 +1,16 @@
 // Copyright Epic Games, Inc. All Rights Reserved.
 
+import 'package:epic_common/theme.dart';
+import 'package:epic_common/widgets.dart';
 import 'package:flutter/material.dart';
 import 'package:flutter_gen/gen_l10n/app_localizations.dart';
 import 'package:provider/provider.dart';
 import 'package:streaming_shared_preferences/streaming_shared_preferences.dart';
 
 import '../../models/settings/delta_widget_settings.dart';
-import '../../utilities/unreal_colors.dart';
 import 'delta_slider.dart';
 import 'delta_widget_base.dart';
 import 'dropdown_button.dart';
-import 'dropdown_list_menu.dart';
-import 'epic_icon_button.dart';
 
 const double _minSensitivity = 0.2;
 const double _maxSensitivity = 2.0;
@@ -75,7 +74,7 @@ class _SensitivitySliderButtonState extends State<SensitivitySliderButton> {
           alignment: Alignment.center,
           children: [
             EpicIconButton(
-              iconPath: 'assets/images/icons/slider_speed.svg',
+              iconPath: 'packages/epic_common/assets/icons/slider_speed.svg',
               bIsToggledOn: state != ModalDropdownButtonState.closed,
               bIsVisualOnly: true,
             ),
diff --git a/Engine/Extras/VirtualProduction/EpicStageApp/lib/widgets/elements/stepper.dart b/Engine/Extras/VirtualProduction/EpicStageApp/lib/widgets/elements/stepper.dart
--- a/Engine/Extras/VirtualProduction/EpicStageApp/lib/widgets/elements/stepper.dart
+++ b/Engine/Extras/VirtualProduction/EpicStageApp/lib/widgets/elements/stepper.dart
@@ -2,13 +2,13 @@
 
 import 'dart:ui';
 
+import 'package:epic_common/preferences.dart';
+import 'package:epic_common/theme.dart';
 import 'package:flutter/material.dart';
 import 'package:flutter_gen/gen_l10n/app_localizations.dart';
 import 'package:provider/provider.dart';
 
 import '../../models/settings/delta_widget_settings.dart';
-import '../../utilities/transient_preference.dart';
-import '../../utilities/unreal_colors.dart';
 import 'delta_widget_base.dart';
 import 'unreal_widget_base.dart';
 
diff --git a/Engine/Extras/VirtualProduction/EpicStageApp/lib/widgets/elements/swipe_revealer.dart b/Engine/Extras/VirtualProduction/EpicStageApp/lib/widgets/elements/swipe_revealer.dart
deleted file mode 100644
diff --git a/Engine/Extras/VirtualProduction/EpicStageApp/lib/widgets/elements/template_menu.dart b/Engine/Extras/VirtualProduction/EpicStageApp/lib/widgets/elements/template_menu.dart
--- a/Engine/Extras/VirtualProduction/EpicStageApp/lib/widgets/elements/template_menu.dart
+++ b/Engine/Extras/VirtualProduction/EpicStageApp/lib/widgets/elements/template_menu.dart
@@ -2,6 +2,8 @@
 
 import 'dart:async';
 
+import 'package:epic_common/localizations.dart';
+import 'package:epic_common/widgets.dart';
 import 'package:flutter/material.dart';
 import 'package:flutter_gen/gen_l10n/app_localizations.dart';
 import 'package:provider/provider.dart';
@@ -9,11 +11,6 @@ import 'package:provider/provider.dart';
 import '../../models/engine_connection.dart';
 import '../../models/unreal_actor_creator.dart';
 import '../../models/unreal_types.dart';
-import '../search_bar.dart';
-import 'epic_icon_button.dart';
-import 'epic_list_view.dart';
-import 'layout/card.dart';
-import 'list_menu.dart';
 
 /// Drop-down menu shown to select a template to create an actor from.
 class TemplatePickerMenu extends StatefulWidget {
@@ -193,11 +190,11 @@ class _TemplatePickerMenuState extends State<TemplatePickerMenu> {
             minWidth: 270,
             children: [
               CardLargeHeader(
-                iconPath: 'assets/images/icons/viewport.svg',
+                iconPath: 'packages/epic_common/assets/icons/viewport.svg',
                 title: AppLocalizations.of(context)!.placeActorMenuAllTemplates,
                 trailing: EpicIconButton(
-                  iconPath: 'assets/images/icons/close.svg',
-                  tooltipMessage: AppLocalizations.of(context)!.menuButtonCancel,
+                  iconPath: 'packages/epic_common/assets/icons/close.svg',
+                  tooltipMessage: EpicCommonLocalizations.of(context)!.menuButtonCancel,
                   iconSize: 20,
                   onPressed: _closeMenu,
                 ),
@@ -247,7 +244,7 @@ class _TemplateMenuItem extends StatelessWidget {
   Widget build(BuildContext context) {
     return ListMenuSimpleItem(
       title: templateData.name,
-      iconPath: 'assets/images/icons/template.svg',
+      iconPath: 'packages/epic_common/assets/icons/template.svg',
       onTap: onPressed,
     );
   }
diff --git a/Engine/Extras/VirtualProduction/EpicStageApp/lib/widgets/elements/tree_view.dart b/Engine/Extras/VirtualProduction/EpicStageApp/lib/widgets/elements/tree_view.dart
deleted file mode 100644
diff --git a/Engine/Extras/VirtualProduction/EpicStageApp/lib/widgets/screens/connect/connect.dart b/Engine/Extras/VirtualProduction/EpicStageApp/lib/widgets/screens/connect/connect.dart
--- a/Engine/Extras/VirtualProduction/EpicStageApp/lib/widgets/screens/connect/connect.dart
+++ b/Engine/Extras/VirtualProduction/EpicStageApp/lib/widgets/screens/connect/connect.dart
@@ -1,12 +1,11 @@
 // Copyright Epic Games, Inc. All Rights Reserved.
 
+import 'package:epic_common/theme.dart';
+import 'package:epic_common/widgets.dart';
 import 'package:flutter/material.dart';
 import 'package:flutter_gen/gen_l10n/app_localizations.dart';
 import 'package:provider/provider.dart';
 
-import '../../../../utilities/constants.dart';
-import '../../../../widgets/elements/empty_placeholder.dart';
-import '../../../../widgets/elements/layout/card.dart';
 import '../../../../widgets/screens/connect/views/connection_list.dart';
 import '../../../../widgets/screens/main/toolbar/main_screen_toolbar.dart';
 import '../../../models/engine_connection.dart';
@@ -14,9 +13,6 @@ import '../../../models/navigator_keys.dart';
 import '../../../models/unreal_engine_beacon.dart';
 import '../../../utilities/debug_utilities.dart';
 import '../../../utilities/net_utilities.dart';
-import '../../../utilities/unreal_colors.dart';
-import '../../elements/epic_icon_button.dart';
-import '../../elements/modal.dart';
 import 'mixins/connect_mixin.dart';
 import 'views/manual_connect_form.dart';
 import 'views/quick_action_grid.dart';
@@ -102,13 +98,17 @@ class _ConnectScreenState extends State<ConnectScreen> with WidgetsBindingObserv
       appBar: const _ConnectScreenToolbar(),
       body: Container(
         color: UnrealColors.gray14,
-        padding: EdgeInsets.only(left: cardMargin, right: cardMargin, bottom: cardMargin),
+        padding: EdgeInsets.only(
+          left: UnrealTheme.cardMargin,
+          right: UnrealTheme.cardMargin,
+          bottom: UnrealTheme.cardMargin,
+        ),
         child: Container(
           width: MediaQuery.of(context).size.width,
-          padding: EdgeInsets.all(cardMargin),
+          padding: EdgeInsets.all(UnrealTheme.cardMargin),
           decoration: BoxDecoration(
             color: Theme.of(context).colorScheme.background,
-            borderRadius: BorderRadius.circular(outerCornerRadius),
+            borderRadius: BorderRadius.circular(UnrealTheme.outerCornerRadius),
           ),
           child: Row(
             mainAxisSize: MainAxisSize.min,
@@ -119,14 +119,14 @@ class _ConnectScreenState extends State<ConnectScreen> with WidgetsBindingObserv
                     children: [
                       CardLargeHeader(
                         title: AppLocalizations.of(context)!.connectScreenQuickConnectPanelTitle,
-                        iconPath: 'assets/images/icons/light_card.svg',
+                        iconPath: 'packages/epic_common/assets/icons/light_card.svg',
                       ),
                       QuickActionGrid(recentConnections: _connections),
                     ],
                   ),
                 ),
               ),
-              SizedBox(width: cardMargin),
+              SizedBox(width: UnrealTheme.cardMargin),
               Expanded(
                 child: Card(
                   child: Column(
@@ -134,7 +134,7 @@ class _ConnectScreenState extends State<ConnectScreen> with WidgetsBindingObserv
                     children: [
                       CardLargeHeader(
                         title: AppLocalizations.of(context)!.connectScreenAllConnectionsPanelTitle,
-                        iconPath: 'assets/images/icons/unreal_u_logo.svg',
+                        iconPath: 'packages/epic_common/assets/icons/unreal_u_logo.svg',
                       ),
                       Expanded(
                         child: _engineBeacon.connections.length > 0
@@ -145,7 +145,7 @@ class _ConnectScreenState extends State<ConnectScreen> with WidgetsBindingObserv
                                 message: AppLocalizations.of(context)!.connectScreenAllConnectionsPanelEmptyMessage,
                                 button: EpicWideButton(
                                   text: AppLocalizations.of(context)!.connectScreenAllConnectionsPanelEmptyButtonLabel,
-                                  iconPath: 'assets/images/icons/plus.svg',
+                                  iconPath: 'packages/epic_common/assets/icons/plus.svg',
                                   color: UnrealColors.highlightGreen,
                                   onPressed: _showManualConnectDialog,
                                 ),
diff --git a/Engine/Extras/VirtualProduction/EpicStageApp/lib/widgets/screens/connect/mixins/connect_mixin.dart b/Engine/Extras/VirtualProduction/EpicStageApp/lib/widgets/screens/connect/mixins/connect_mixin.dart
--- a/Engine/Extras/VirtualProduction/EpicStageApp/lib/widgets/screens/connect/mixins/connect_mixin.dart
+++ b/Engine/Extras/VirtualProduction/EpicStageApp/lib/widgets/screens/connect/mixins/connect_mixin.dart
@@ -2,6 +2,9 @@
 
 import 'dart:async';
 
+import 'package:epic_common/localizations.dart';
+import 'package:epic_common/preferences.dart';
+import 'package:epic_common/widgets.dart';
 import 'package:flutter/material.dart';
 import 'package:flutter_gen/gen_l10n/app_localizations.dart';
 import 'package:provider/provider.dart';
@@ -13,10 +16,6 @@ import '../../../../models/unreal_actor_manager.dart';
 import '../../../../models/unreal_types.dart';
 import '../../../../utilities/constants.dart';
 import '../../../../utilities/net_utilities.dart';
-import '../../../../utilities/preferences_bundle.dart';
-import '../../../elements/asset_icon.dart';
-import '../../../elements/epic_icon_button.dart';
-import '../../../elements/modal.dart';
 import '../../main/stage_app_main_screen.dart';
 import '../views/error_modal.dart';
 import '../views/n_display_selector_dialog.dart';
@@ -130,7 +129,7 @@ mixin ConnectMixin<T extends StatefulWidget> on State<T> {
                 mainAxisSize: MainAxisSize.min,
                 children: [
                   AssetIcon(
-                    path: 'assets/images/icons/unreal_u_logo.svg',
+                    path: 'packages/epic_common/assets/icons/unreal_u_logo.svg',
                     size: 24,
                   ),
                   SizedBox(width: 16),
@@ -151,7 +150,7 @@ mixin ConnectMixin<T extends StatefulWidget> on State<T> {
               ),
               Row(mainAxisAlignment: MainAxisAlignment.end, children: [
                 EpicLozengeButton(
-                  label: AppLocalizations.of(context)!.menuButtonCancel,
+                  label: EpicCommonLocalizations.of(context)!.menuButtonCancel,
                   color: Theme.of(context).colorScheme.secondary,
                   onPressed: () {
                     Navigator.pop(context, false);
diff --git a/Engine/Extras/VirtualProduction/EpicStageApp/lib/widgets/screens/connect/views/connection_tile.dart b/Engine/Extras/VirtualProduction/EpicStageApp/lib/widgets/screens/connect/views/connection_tile.dart
--- a/Engine/Extras/VirtualProduction/EpicStageApp/lib/widgets/screens/connect/views/connection_tile.dart
+++ b/Engine/Extras/VirtualProduction/EpicStageApp/lib/widgets/screens/connect/views/connection_tile.dart
@@ -1,12 +1,11 @@
 // Copyright Epic Games, Inc. All Rights Reserved.
 
+import 'package:epic_common/theme.dart';
+import 'package:epic_common/widgets.dart';
 import 'package:flutter/material.dart';
 import 'package:flutter_gen/gen_l10n/app_localizations.dart';
 
-import '../../../../utilities/constants.dart';
 import '../../../../utilities/net_utilities.dart';
-import '../../../../utilities/unreal_colors.dart';
-import '../../../elements/asset_icon.dart';
 import '../mixins/connect_mixin.dart';
 
 /// Connection tile/item used to visually represent available instances of UE connections.
@@ -63,7 +62,7 @@ class _ConnectionTileState extends State<ConnectionTile> with ConnectMixin {
         height: 166,
         padding: EdgeInsets.symmetric(horizontal: 10),
         decoration: BoxDecoration(
-          borderRadius: BorderRadius.circular(cardCornerRadius),
+          borderRadius: BorderRadius.circular(UnrealTheme.cardCornerRadius),
           color: fillColor,
           border: widget.bFilled
               ? null
@@ -77,7 +76,7 @@ class _ConnectionTileState extends State<ConnectionTile> with ConnectMixin {
           mainAxisAlignment: MainAxisAlignment.center,
           children: [
             AssetIcon(
-              path: 'assets/images/icons/unreal_u.svg',
+              path: 'packages/epic_common/assets/icons/unreal_u.svg',
               size: 50,
               color: widget.bFilled ? UnrealColors.white : UnrealColors.gray75,
             ),
diff --git a/Engine/Extras/VirtualProduction/EpicStageApp/lib/widgets/screens/connect/views/error_modal.dart b/Engine/Extras/VirtualProduction/EpicStageApp/lib/widgets/screens/connect/views/error_modal.dart
--- a/Engine/Extras/VirtualProduction/EpicStageApp/lib/widgets/screens/connect/views/error_modal.dart
+++ b/Engine/Extras/VirtualProduction/EpicStageApp/lib/widgets/screens/connect/views/error_modal.dart
@@ -1,12 +1,11 @@
 // Copyright Epic Games, Inc. All Rights Reserved.
 
+import 'package:epic_common/localizations.dart';
+import 'package:epic_common/widgets.dart';
 import 'package:flutter/material.dart';
 import 'package:flutter_gen/gen_l10n/app_localizations.dart';
 
 import '../../../../utilities/net_utilities.dart';
-import '../../../elements/asset_icon.dart';
-import '../../../elements/epic_icon_button.dart';
-import '../../../elements/parsed_rich_text.dart';
 
 /// Error modal for when there's an error connecting to an instance of UE.
 class ErrorModal extends StatelessWidget {
@@ -30,7 +29,7 @@ class ErrorModal extends StatelessWidget {
             mainAxisAlignment: MainAxisAlignment.center,
             children: [
               AssetIcon(
-                path: 'assets/images/icons/alert_triangle_large.svg',
+                path: 'packages/epic_common/assets/icons/alert_triangle_large.svg',
                 size: 24,
               ),
               SizedBox(width: 10),
@@ -55,7 +54,7 @@ class ErrorModal extends StatelessWidget {
             children: [
               EpicLozengeButton(
                 onPressed: () => Navigator.of(context).pop(),
-                label: AppLocalizations.of(context)!.menuButtonCancel,
+                label: EpicCommonLocalizations.of(context)!.menuButtonCancel,
                 color: Colors.transparent,
               ),
               SizedBox(width: 16),
@@ -64,7 +63,7 @@ class ErrorModal extends StatelessWidget {
                   Navigator.of(context).pop();
                   reconnect?.call();
                 },
-                label: AppLocalizations.of(context)!.menuButtonRetry,
+                label: EpicCommonLocalizations.of(context)!.menuButtonRetry,
               ),
             ],
           ),
diff --git a/Engine/Extras/VirtualProduction/EpicStageApp/lib/widgets/screens/connect/views/manual_connect_form.dart b/Engine/Extras/VirtualProduction/EpicStageApp/lib/widgets/screens/connect/views/manual_connect_form.dart
--- a/Engine/Extras/VirtualProduction/EpicStageApp/lib/widgets/screens/connect/views/manual_connect_form.dart
+++ b/Engine/Extras/VirtualProduction/EpicStageApp/lib/widgets/screens/connect/views/manual_connect_form.dart
@@ -2,14 +2,14 @@
 
 import 'dart:io';
 
+import 'package:epic_common/localizations.dart';
+import 'package:epic_common/widgets.dart';
 import 'package:flutter/material.dart';
 import 'package:flutter_gen/gen_l10n/app_localizations.dart';
 import 'package:uuid/uuid.dart';
 
 import '../../../../utilities/constants.dart';
 import '../../../../utilities/net_utilities.dart';
-import '../../../elements/epic_icon_button.dart';
-import '../../../elements/modal.dart';
 
 /// Manual connection form for manually connecting to an instance of UE.
 class ManualConnectForm extends StatefulWidget {
@@ -76,12 +76,12 @@ class _ManualConnectFormState extends State<ManualConnectForm> {
                   mainAxisAlignment: MainAxisAlignment.end,
                   children: [
                     EpicLozengeButton(
-                      label: AppLocalizations.of(context)!.menuButtonCancel,
+                      label: EpicCommonLocalizations.of(context)!.menuButtonCancel,
                       color: Colors.transparent,
                       onPressed: () => Navigator.of(context).pop(),
                     ),
                     EpicLozengeButton(
-                      label: AppLocalizations.of(context)!.menuButtonOK,
+                      label: EpicCommonLocalizations.of(context)!.menuButtonOK,
                       onPressed: _onConnect,
                     ),
                   ],
diff --git a/Engine/Extras/VirtualProduction/EpicStageApp/lib/widgets/screens/connect/views/manual_connect_tile.dart b/Engine/Extras/VirtualProduction/EpicStageApp/lib/widgets/screens/connect/views/manual_connect_tile.dart
--- a/Engine/Extras/VirtualProduction/EpicStageApp/lib/widgets/screens/connect/views/manual_connect_tile.dart
+++ b/Engine/Extras/VirtualProduction/EpicStageApp/lib/widgets/screens/connect/views/manual_connect_tile.dart
@@ -2,13 +2,12 @@
 
 import 'dart:async';
 
+import 'package:epic_common/theme.dart';
+import 'package:epic_common/widgets.dart';
 import 'package:flutter/material.dart';
 import 'package:flutter_gen/gen_l10n/app_localizations.dart';
 
 import '../../../../models/navigator_keys.dart';
-import '../../../../utilities/constants.dart';
-import '../../../elements/asset_icon.dart';
-import '../../../elements/modal.dart';
 import '../mixins/connect_mixin.dart';
 import 'manual_connect_form.dart';
 
@@ -41,7 +40,7 @@ class _ManualConnectTileState extends State<ManualConnectTile> with ConnectMixin
         width: 170,
         height: 166,
         decoration: BoxDecoration(
-          borderRadius: BorderRadius.circular(cardCornerRadius),
+          borderRadius: BorderRadius.circular(UnrealTheme.cardCornerRadius),
           color: Theme.of(context).colorScheme.secondary,
         ),
         padding: EdgeInsets.all(10),
@@ -51,7 +50,7 @@ class _ManualConnectTileState extends State<ManualConnectTile> with ConnectMixin
           mainAxisSize: MainAxisSize.min,
           children: [
             AssetIcon(
-              path: 'assets/images/icons/unreal_u.svg',
+              path: 'packages/epic_common/assets/icons/unreal_u.svg',
               size: 50,
             ),
             SizedBox(height: 16),
diff --git a/Engine/Extras/VirtualProduction/EpicStageApp/lib/widgets/screens/connect/views/n_display_selector_dialog.dart b/Engine/Extras/VirtualProduction/EpicStageApp/lib/widgets/screens/connect/views/n_display_selector_dialog.dart
--- a/Engine/Extras/VirtualProduction/EpicStageApp/lib/widgets/screens/connect/views/n_display_selector_dialog.dart
+++ b/Engine/Extras/VirtualProduction/EpicStageApp/lib/widgets/screens/connect/views/n_display_selector_dialog.dart
@@ -1,5 +1,8 @@
 // Copyright Epic Games, Inc. All Rights Reserved.
 
+import 'package:epic_common/localizations.dart';
+import 'package:epic_common/theme.dart';
+import 'package:epic_common/widgets.dart';
 import 'package:flutter/material.dart';
 import 'package:flutter_gen/gen_l10n/app_localizations.dart';
 import 'package:provider/provider.dart';
@@ -7,10 +10,6 @@ import 'package:provider/provider.dart';
 import '../../../../models/unreal_actor_manager.dart';
 import '../../../../models/unreal_types.dart';
 import '../../../../utilities/constants.dart';
-import '../../../../utilities/unreal_colors.dart';
-import '../../../elements/asset_icon.dart';
-import '../../../elements/epic_icon_button.dart';
-import '../../../elements/modal.dart';
 
 /// Displays the list of currently available actors which users can select from.
 class NDisplaySelectorDialog extends StatefulWidget {
@@ -84,7 +83,7 @@ class _NDisplaySelectorDialogState extends State<NDisplaySelectorDialog> {
               mainAxisAlignment: MainAxisAlignment.center,
               children: [
                 AssetIcon(
-                  path: 'assets/images/icons/unreal_u_logo.svg',
+                  path: 'packages/epic_common/assets/icons/unreal_u_logo.svg',
                   size: 24,
                 ),
                 SizedBox(width: 20),
@@ -115,7 +114,7 @@ class _NDisplaySelectorDialogState extends State<NDisplaySelectorDialog> {
                     if (!widget.bIsAlreadyConnected)
                       EpicLozengeButton(
                         onPressed: () => Navigator.pop(context),
-                        label: AppLocalizations.of(context)!.menuButtonCancel,
+                        label: EpicCommonLocalizations.of(context)!.menuButtonCancel,
                         color: (_selectedActor != null) ? Colors.transparent : Theme.of(context).colorScheme.secondary,
                       ),
                     if (_selectedActor != null)
@@ -124,7 +123,7 @@ class _NDisplaySelectorDialogState extends State<NDisplaySelectorDialog> {
                         child: EpicLozengeButton(
                             onPressed: () => Navigator.of(context).pop(_selectedActor),
                             label: widget.bIsAlreadyConnected
-                                ? AppLocalizations.of(context)!.menuButtonProceed
+                                ? EpicCommonLocalizations.of(context)!.menuButtonProceed
                                 : AppLocalizations.of(context)!.nDisplaySelectorModalConnectButtonLabel),
                       ),
                   ],
@@ -188,7 +187,7 @@ class NDisplayItem extends StatelessWidget {
         child: Row(
           children: [
             AssetIcon(
-              path: 'assets/images/icons/ndisplay.svg',
+              path: 'packages/epic_common/assets/icons/ndisplay.svg',
               size: 32,
             ),
             SizedBox(width: 16),
diff --git a/Engine/Extras/VirtualProduction/EpicStageApp/lib/widgets/screens/connect/views/passphrase_modal.dart b/Engine/Extras/VirtualProduction/EpicStageApp/lib/widgets/screens/connect/views/passphrase_modal.dart
--- a/Engine/Extras/VirtualProduction/EpicStageApp/lib/widgets/screens/connect/views/passphrase_modal.dart
+++ b/Engine/Extras/VirtualProduction/EpicStageApp/lib/widgets/screens/connect/views/passphrase_modal.dart
@@ -1,5 +1,8 @@
 // Copyright Epic Games, Inc. All Rights Reserved.
 
+import 'package:epic_common/localizations.dart';
+import 'package:epic_common/theme.dart';
+import 'package:epic_common/widgets.dart';
 import 'package:flutter/material.dart';
 import 'package:flutter_gen/gen_l10n/app_localizations.dart';
 import 'package:logging/logging.dart';
@@ -8,10 +11,6 @@ import 'package:provider/provider.dart';
 import '../../../../models/engine_connection.dart';
 import '../../../../models/engine_passphrase_manager.dart';
 import '../../../../utilities/net_utilities.dart';
-import '../../../../utilities/unreal_colors.dart';
-import '../../../elements/epic_checkbox.dart';
-import '../../../elements/epic_icon_button.dart';
-import '../../../elements/modal.dart';
 import '../mixins/connect_mixin.dart';
 
 final _log = Logger('PassphraseModal');
@@ -113,13 +112,13 @@ class _PassphraseModalDialogState extends State<PassphraseModalDialog> with Conn
                         mainAxisAlignment: MainAxisAlignment.spaceAround,
                         children: [
                           EpicLozengeButton(
-                            label: localizations.menuButtonCancel,
+                            label: EpicCommonLocalizations.of(context)!.menuButtonCancel,
                             width: buttonWidth,
                             color: Colors.transparent,
                             onPressed: _cancel,
                           ),
                           EpicLozengeButton(
-                            label: localizations.menuButtonSubmit,
+                            label: EpicCommonLocalizations.of(context)!.menuButtonSubmit,
                             width: buttonWidth,
                             onPressed: _submit,
                           ),
diff --git a/Engine/Extras/VirtualProduction/EpicStageApp/lib/widgets/screens/eula/eula_screen.dart b/Engine/Extras/VirtualProduction/EpicStageApp/lib/widgets/screens/eula/eula_screen.dart
--- a/Engine/Extras/VirtualProduction/EpicStageApp/lib/widgets/screens/eula/eula_screen.dart
+++ b/Engine/Extras/VirtualProduction/EpicStageApp/lib/widgets/screens/eula/eula_screen.dart
@@ -1,14 +1,13 @@
 // Copyright Epic Games, Inc. All Rights Reserved.
 
+import 'package:epic_common/localizations.dart';
+import 'package:epic_common/preferences.dart';
+import 'package:epic_common/theme.dart';
+import 'package:epic_common/widgets.dart';
 import 'package:flutter/material.dart';
 import 'package:flutter_gen/gen_l10n/app_localizations.dart';
 import 'package:provider/provider.dart';
 
-import '../../../utilities/preferences_bundle.dart';
-import '../../../utilities/unreal_colors.dart';
-import '../../elements/epic_icon_button.dart';
-import '../../elements/epic_scroll_view.dart';
-import '../../elements/parsed_rich_text.dart';
 import '../connect/connect.dart';
 
 /// Data types for EULA content.
@@ -76,7 +75,7 @@ class _EulaScreenState extends State<EulaScreen> {
           color: UnrealColors.highlightBlue,
           child: Text(
             bHasAcceptedEula
-                ? AppLocalizations.of(context)!.menuButtonOK
+                ? EpicCommonLocalizations.of(context)!.menuButtonOK
                 : AppLocalizations.of(context)!.eulaAcceptButtonLabel,
           ),
           onPressed: () {
diff --git a/Engine/Extras/VirtualProduction/EpicStageApp/lib/widgets/screens/main/sidebar/outliner_filter_menu.dart b/Engine/Extras/VirtualProduction/EpicStageApp/lib/widgets/screens/main/sidebar/outliner_filter_menu.dart
--- a/Engine/Extras/VirtualProduction/EpicStageApp/lib/widgets/screens/main/sidebar/outliner_filter_menu.dart
+++ b/Engine/Extras/VirtualProduction/EpicStageApp/lib/widgets/screens/main/sidebar/outliner_filter_menu.dart
@@ -1,5 +1,6 @@
 // Copyright Epic Games, Inc. All Rights Reserved.
 
+import 'package:epic_common/widgets.dart';
 import 'package:flutter/material.dart';
 import 'package:flutter_gen/gen_l10n/app_localizations.dart';
 import 'package:provider/provider.dart';
@@ -8,9 +9,6 @@ import 'package:streaming_shared_preferences/streaming_shared_preferences.dart';
 import '../../../../models/settings/outliner_panel_settings.dart';
 import '../../../../models/unreal_object_filters.dart';
 import '../../../elements/dropdown_button.dart';
-import '../../../elements/dropdown_list_menu.dart';
-import '../../../elements/layout/card.dart';
-import '../../../elements/list_menu.dart';
 
 /// Modal dropdown button that displays outliner filter settings.
 class OutlinerFilterButton extends StatelessWidget {
@@ -26,7 +24,9 @@ class OutlinerFilterButton extends StatelessWidget {
         child: PreferenceBuilder(
           preference: settings.selectedFilters,
           builder: (context, Set<UnrealObjectFilter> filters) => CardSubHeaderButton(
-            iconPath: filters.isEmpty ? 'assets/images/icons/filter.svg' : 'assets/images/icons/filter_on.svg',
+            iconPath: filters.isEmpty
+                ? 'packages/epic_common/assets/icons/filter.svg'
+                : 'packages/epic_common/assets/icons/filter_on.svg',
             tooltipMessage: AppLocalizations.of(context)!.outlinerFiltersButtonTooltip,
             bIsToggledOn: state != ModalDropdownButtonState.closed,
             bIsVisualOnly: true,
diff --git a/Engine/Extras/VirtualProduction/EpicStageApp/lib/widgets/screens/main/sidebar/outliner_panel.dart b/Engine/Extras/VirtualProduction/EpicStageApp/lib/widgets/screens/main/sidebar/outliner_panel.dart
--- a/Engine/Extras/VirtualProduction/EpicStageApp/lib/widgets/screens/main/sidebar/outliner_panel.dart
+++ b/Engine/Extras/VirtualProduction/EpicStageApp/lib/widgets/screens/main/sidebar/outliner_panel.dart
@@ -3,6 +3,9 @@
 import 'dart:async';
 
 import 'package:async/async.dart';
+import 'package:epic_common/preferences.dart';
+import 'package:epic_common/theme.dart';
+import 'package:epic_common/widgets.dart';
 import 'package:flutter/material.dart';
 import 'package:flutter/services.dart';
 import 'package:flutter_gen/gen_l10n/app_localizations.dart';
@@ -23,19 +26,8 @@ import '../../../../models/unreal_transaction_manager.dart';
 import '../../../../models/unreal_types.dart';
 import '../../../../utilities/constants.dart';
 import '../../../../utilities/guarded_refresh_state.dart';
-import '../../../../utilities/preferences_bundle.dart';
-import '../../../../utilities/transient_preference.dart';
-import '../../../../utilities/unreal_colors.dart';
 import '../../../elements/dropdown_button.dart';
-import '../../../elements/dropdown_list_menu.dart';
-import '../../../elements/empty_placeholder.dart';
-import '../../../elements/epic_icon_button.dart';
-import '../../../elements/layout/card.dart';
-import '../../../elements/list_menu.dart';
-import '../../../elements/modal.dart';
 import '../../../elements/place_actor_menu.dart';
-import '../../../elements/swipe_revealer.dart';
-import '../../../search_bar.dart';
 import '../stage_app_main_screen.dart';
 import 'outliner_filter_menu.dart';
 
@@ -268,7 +260,7 @@ class _OutlinerPanelState extends State<OutlinerPanel> {
                 CardSmallHeader(title: AppLocalizations.of(context)!.outlinerTitle),
 
                 Container(
-                  height: sectionMargin,
+                  height: UnrealTheme.sectionMargin,
                   color: Theme.of(context).colorScheme.surfaceVariant,
                 ),
 
@@ -293,7 +285,7 @@ class _OutlinerPanelState extends State<OutlinerPanel> {
                         preference: selectedActorSettings.bIsInMultiSelectMode,
                         builder: (BuildContext context, bool bIsInMultiSelectMode) {
                           return CardSubHeaderButton(
-                            iconPath: 'assets/images/icons/multi_select.svg',
+                            iconPath: 'packages/epic_common/assets/icons/multi_select.svg',
                             tooltipMessage: AppLocalizations.of(context)!.toggleMultiselectTooltip,
                             bIsToggledOn: bIsInMultiSelectMode,
                             onPressed: () => selectedActorSettings.bIsInMultiSelectMode.setValue(!bIsInMultiSelectMode),
@@ -325,7 +317,7 @@ class _OutlinerPanelState extends State<OutlinerPanel> {
                             buttonBuilder: (context, state) => Material(
                               color: Colors.transparent,
                               child: CardSubHeaderButton(
-                                iconPath: 'assets/images/icons/ellipsis.svg',
+                                iconPath: 'packages/epic_common/assets/icons/ellipsis.svg',
                                 tooltipMessage: AppLocalizations.of(context)!.moreActions,
                                 bIsToggledOn: state != ModalDropdownButtonState.closed,
                                 bIsVisualOnly: selectedActors.isNotEmpty,
@@ -336,7 +328,7 @@ class _OutlinerPanelState extends State<OutlinerPanel> {
                               children: [
                                 if (bCanFocus)
                                   ListMenuSimpleItem(
-                                    iconPath: 'assets/images/icons/focus.svg',
+                                    iconPath: 'packages/epic_common/assets/icons/focus.svg',
                                     title: AppLocalizations.of(context)!.outlinerFocusSelected,
                                     onTap: () {
                                       _focusSelectedActor();
@@ -345,7 +337,7 @@ class _OutlinerPanelState extends State<OutlinerPanel> {
                                     bIsEnabled: bCanFocus,
                                   ),
                                 ListMenuSimpleItem(
-                                  iconPath: 'assets/images/icons/paste.svg',
+                                  iconPath: 'packages/epic_common/assets/icons/paste.svg',
                                   title: AppLocalizations.of(context)!.outlinerDuplicateSelected,
                                   onTap: () {
                                     _duplicateSelectedActors();
@@ -353,7 +345,7 @@ class _OutlinerPanelState extends State<OutlinerPanel> {
                                   },
                                 ),
                                 ListMenuSimpleItem(
-                                  iconPath: 'assets/images/icons/trash.svg',
+                                  iconPath: 'packages/epic_common/assets/icons/trash.svg',
                                   title: AppLocalizations.of(context)!.outlinerDeleteSelected,
                                   onTap: () {
                                     _deleteSelectedActors();
@@ -362,7 +354,7 @@ class _OutlinerPanelState extends State<OutlinerPanel> {
                                 ),
                                 if (bHasSingleTarget)
                                   ListMenuSimpleItem(
-                                    iconPath: 'assets/images/icons/edit.svg',
+                                    iconPath: 'packages/epic_common/assets/icons/edit.svg',
                                     title: AppLocalizations.of(context)!.outlinerRenameSelected,
                                     onTap: () {
                                       Navigator.of(context).pop();
@@ -383,7 +375,7 @@ class _OutlinerPanelState extends State<OutlinerPanel> {
                 Expanded(
                   child: _filteredActors.isNotEmpty
                       ? ListView.builder(
-                          padding: cardListViewPadding,
+                          padding: UnrealTheme.cardListViewPadding,
                           itemCount: _filteredActors.length,
                           itemBuilder: (BuildContext context, int actorIndex) {
                             final UnrealObject actor = _filteredActors[actorIndex];
@@ -402,9 +394,10 @@ class _OutlinerPanelState extends State<OutlinerPanel> {
                           message: AppLocalizations.of(context)!.outlinerEmptyMessage,
                           button: EpicWideButton(
                             text: AppLocalizations.of(context)!.addActorButton,
-                            iconPath: 'assets/images/icons/plus.svg',
+                            iconPath: 'packages/epic_common/assets/icons/plus.svg',
                             iconColor: UnrealColors.highlightGreen,
                             onPressed: () => DropDownListMenu.showAtWidget(
+                              context,
                               widgetKey:
                                   Provider.of<StageAppMainScreenKeys>(context, listen: false).placeActorsButtonKey,
                               builder: (context) => PlaceActorDropDownMenu(
@@ -724,14 +717,14 @@ class _OutlinerPanelActorState extends State<_OutlinerPanelActor> with GuardedRe
           backgroundPadding: const EdgeInsets.only(top: 2),
           onDeleted: () => _deleteActors(context, {widget.actor.path}),
           leftSwipeActionBuilder: (context, onFinished) => CardListTileSwipeAction(
-            iconPath: 'assets/images/icons/trash.svg',
+            iconPath: 'packages/epic_common/assets/icons/trash.svg',
             color: UnrealColors.highlightRed,
             onPressed: () => onFinished(bDeleteItem: true),
           ),
           rightSwipeActionBuilder: (context, onFinished) => CardListTileSwipeAction(
             iconPath: (_actorData?.bIsHiddenInGame == true)
-                ? 'assets/images/icons/hidden_in_game.svg'
-                : 'assets/images/icons/visible_in_game.svg',
+                ? 'packages/epic_common/assets/icons/hidden_in_game.svg'
+                : 'packages/epic_common/assets/icons/visible_in_game.svg',
             color: UnrealColors.gray22,
             onPressed: () {
               if (_actorData != null) {
@@ -818,7 +811,7 @@ class OutlinerToggleButton extends StatelessWidget {
       preference: mainScreenSettings.bIsOutlinerPanelOpen,
       builder: (BuildContext context, bool bIsOutlinerPanelOpen) {
         return EpicIconButton(
-          iconPath: 'assets/images/icons/outliner.svg',
+          iconPath: 'packages/epic_common/assets/icons/outliner.svg',
           tooltipMessage: AppLocalizations.of(context)!.outlinerTitle,
           bIsToggledOn: bIsOutlinerPanelOpen,
           onPressed: bEnabled
@@ -861,7 +854,9 @@ class _LightCardVisibilityToggleState extends State<LightCardVisibilityToggle> w
       preference: selectedActorSettings.selectedActors,
       builder: (BuildContext context, Set<String> selectedActors) {
         return CardSubHeaderButton(
-          iconPath: bIsVisible ? 'assets/images/icons/visible_in_game.svg' : 'assets/images/icons/hidden_in_game.svg',
+          iconPath: bIsVisible
+              ? 'packages/epic_common/assets/icons/visible_in_game.svg'
+              : 'packages/epic_common/assets/icons/hidden_in_game.svg',
           tooltipMessage: AppLocalizations.of(context)!.outlinerToggleVisibility,
           bIsToggledOn: false,
           onPressed: selectedActors.isEmpty ? null : () => _toggleVisibility(),
diff --git a/Engine/Extras/VirtualProduction/EpicStageApp/lib/widgets/screens/main/stage_app_main_screen.dart b/Engine/Extras/VirtualProduction/EpicStageApp/lib/widgets/screens/main/stage_app_main_screen.dart
--- a/Engine/Extras/VirtualProduction/EpicStageApp/lib/widgets/screens/main/stage_app_main_screen.dart
+++ b/Engine/Extras/VirtualProduction/EpicStageApp/lib/widgets/screens/main/stage_app_main_screen.dart
@@ -1,5 +1,7 @@
 // Copyright Epic Games, Inc. All Rights Reserved.
 
+import 'package:epic_common/theme.dart';
+import 'package:epic_common/widgets.dart';
 import 'package:flutter/material.dart';
 import 'package:flutter_gen/gen_l10n/app_localizations.dart';
 import 'package:provider/provider.dart';
@@ -10,9 +12,7 @@ import '../../../models/unreal_actor_manager.dart';
 import '../../../models/unreal_types.dart';
 import '../../../utilities/constants.dart';
 import '../../../utilities/debug_utilities.dart';
-import '../../../utilities/unreal_colors.dart';
 import '../../elements/floating_map_preview.dart';
-import '../../elements/layout/lazy_tab_view.dart';
 import '../connect/views/n_display_selector_dialog.dart';
 import 'tabs/main_screen_tabs.dart';
 import 'toolbar/main_screen_toolbar.dart';
@@ -103,12 +103,16 @@ class _StageAppMainScreenState extends State<StageAppMainScreen>
           children: [
             Container(
               color: UnrealColors.gray14,
-              padding: EdgeInsets.only(left: cardMargin, right: cardMargin, bottom: cardMargin),
+              padding: EdgeInsets.only(
+                left: UnrealTheme.cardMargin,
+                right: UnrealTheme.cardMargin,
+                bottom: UnrealTheme.cardMargin,
+              ),
               key: Key('Tab View'),
               child: Container(
                 decoration: BoxDecoration(
                   color: Theme.of(context).colorScheme.background,
-                  borderRadius: BorderRadius.circular(outerCornerRadius),
+                  borderRadius: BorderRadius.circular(UnrealTheme.outerCornerRadius),
                 ),
                 clipBehavior: Clip.antiAlias,
                 child: LazyTabView(
diff --git a/Engine/Extras/VirtualProduction/EpicStageApp/lib/widgets/screens/main/tabs/base_color_tab.dart b/Engine/Extras/VirtualProduction/EpicStageApp/lib/widgets/screens/main/tabs/base_color_tab.dart
--- a/Engine/Extras/VirtualProduction/EpicStageApp/lib/widgets/screens/main/tabs/base_color_tab.dart
+++ b/Engine/Extras/VirtualProduction/EpicStageApp/lib/widgets/screens/main/tabs/base_color_tab.dart
@@ -1,5 +1,8 @@
 // Copyright Epic Games, Inc. All Rights Reserved.
 
+import 'package:epic_common/preferences.dart';
+import 'package:epic_common/theme.dart';
+import 'package:epic_common/widgets.dart';
 import 'package:flutter/material.dart';
 import 'package:flutter_gen/gen_l10n/app_localizations.dart';
 import 'package:provider/provider.dart';
@@ -8,14 +11,10 @@ import 'package:streaming_shared_preferences/streaming_shared_preferences.dart';
 import '../../../../models/settings/base_color_tab_settings.dart';
 import '../../../../models/settings/color_wheel_settings.dart';
 import '../../../../models/unreal_types.dart';
-import '../../../../utilities/constants.dart';
-import '../../../../utilities/preferences_bundle.dart';
 import '../../../../utilities/unreal_utilities.dart';
 import '../../../elements/color_grading_wheel.dart';
 import '../../../elements/delta_slider.dart';
-import '../../../elements/dropdown_button.dart';
 import '../../../elements/dropdown_text.dart';
-import '../../../elements/selector_bar.dart';
 
 /// Modes the base color tab can operate in.
 enum BaseColorTabMode {
@@ -169,7 +168,7 @@ class _BaseColorTabState extends State<BaseColorTab> {
     return Provider<BaseColorTabSettings>(
       create: (_) => _tabSettings,
       child: Padding(
-        padding: const EdgeInsets.only(top: sectionMargin),
+        padding: const EdgeInsets.only(top: UnrealTheme.sectionMargin),
         child: Row(
           crossAxisAlignment: CrossAxisAlignment.end,
           children: [
@@ -179,14 +178,14 @@ class _BaseColorTabState extends State<BaseColorTab> {
                 child: Column(children: [
                   colorWheel,
                   Padding(
-                    padding: EdgeInsets.symmetric(horizontal: cardMargin),
+                    padding: EdgeInsets.symmetric(horizontal: UnrealTheme.cardMargin),
                     child: miscPropertyControls,
                   ),
                 ]),
               ),
             ),
 
-            const SizedBox(width: sectionMargin),
+            const SizedBox(width: UnrealTheme.sectionMargin),
 
             // Other properties
             Expanded(
@@ -200,7 +199,7 @@ class _BaseColorTabState extends State<BaseColorTab> {
                   ),
                   Expanded(
                     child: Container(
-                      padding: EdgeInsets.only(left: cardMargin, right: cardMargin, top: 16),
+                      padding: EdgeInsets.only(left: UnrealTheme.cardMargin, right: UnrealTheme.cardMargin, top: 16),
                       color: Theme.of(context).colorScheme.surface,
                       child: widget.rightSideContents,
                     ),
@@ -353,7 +352,7 @@ class _BaseColorTabState extends State<BaseColorTab> {
 
                   return UnrealDeltaSlider(
                     overrideName: _getHighlightsRangeDisplayName(highlightsRangeLimitType),
-                    buildLabel: (String name) => DrivenDropdownText(
+                    buildLabel: (String name) => DropdownText(
                       value: highlightsRangeLimitType,
                       items: RangeLimitType.values,
                       makeItemName: (RangeLimitType item) => _getHighlightsRangeDisplayName(item),
diff --git a/Engine/Extras/VirtualProduction/EpicStageApp/lib/widgets/screens/main/tabs/color_grading_tab.dart b/Engine/Extras/VirtualProduction/EpicStageApp/lib/widgets/screens/main/tabs/color_grading_tab.dart
--- a/Engine/Extras/VirtualProduction/EpicStageApp/lib/widgets/screens/main/tabs/color_grading_tab.dart
+++ b/Engine/Extras/VirtualProduction/EpicStageApp/lib/widgets/screens/main/tabs/color_grading_tab.dart
@@ -3,6 +3,9 @@
 import 'dart:async';
 import 'dart:math' as math;
 
+import 'package:epic_common/preferences.dart';
+import 'package:epic_common/theme.dart';
+import 'package:epic_common/widgets.dart';
 import 'package:flutter/material.dart';
 import 'package:flutter_gen/gen_l10n/app_localizations.dart';
 import 'package:logging/logging.dart';
@@ -19,21 +22,11 @@ import '../../../../models/unreal_actor_manager.dart';
 import '../../../../models/unreal_transaction_manager.dart';
 import '../../../../models/unreal_types.dart';
 import '../../../../utilities/constants.dart';
-import '../../../../utilities/preferences_bundle.dart';
-import '../../../../utilities/transient_preference.dart';
-import '../../../../utilities/unreal_colors.dart';
 import '../../../../utilities/unreal_utilities.dart';
 import '../../../elements/delta_slider.dart';
 import '../../../elements/dropdown_text.dart';
-import '../../../elements/empty_placeholder.dart';
-import '../../../elements/epic_list_view.dart';
-import '../../../elements/layout/card.dart';
-import '../../../elements/modal.dart';
 import '../../../elements/reset_mode_button.dart';
-import '../../../elements/selector_bar.dart';
 import '../../../elements/stepper.dart';
-import '../../../elements/swipe_revealer.dart';
-import '../../../elements/tree_view.dart';
 import '../../../elements/unreal_property_builder.dart';
 import '../sidebar/outliner_panel.dart';
 import 'base_color_tab.dart';
@@ -44,7 +37,7 @@ final _log = Logger('ColorGradingTab');
 class ColorGradingTab extends StatefulWidget {
   const ColorGradingTab({Key? key}) : super(key: key);
 
-  static const String iconPath = 'assets/images/icons/color_grading.svg';
+  static const String iconPath = 'packages/epic_common/assets/icons/color_grading.svg';
 
   static String getTitle(BuildContext context) => AppLocalizations.of(context)!.tabTitleColorGrading;
 
@@ -112,7 +105,7 @@ class _ColorGradingTabState extends State<ColorGradingTab> {
             final _ColorGradingTargetListEntryData? targetData = _getTargetEntryDataForProperty(colorGradingTarget);
 
             return Padding(
-              padding: EdgeInsets.all(cardMargin),
+              padding: EdgeInsets.all(UnrealTheme.cardMargin),
               child: Row(children: [
                 // Main controls
                 Expanded(
@@ -144,7 +137,7 @@ class _ColorGradingTabState extends State<ColorGradingTab> {
                     }
 
                     return Row(children: [
-                      SizedBox(width: cardMargin),
+                      SizedBox(width: UnrealTheme.cardMargin),
 
                       // Outliner and target panels
                       SizedBox(
@@ -210,13 +203,13 @@ class _ColorGradingTabState extends State<ColorGradingTab> {
     switch (entryData?.type) {
       case _ColorGradingObjectEntryType.nDisplayConfig:
       case _ColorGradingObjectEntryType.icvfxCamera:
-        return 'assets/images/icons/viewport.svg';
+        return 'packages/epic_common/assets/icons/viewport.svg';
 
       case _ColorGradingObjectEntryType.postProcessVolume:
-        return 'assets/images/icons/post_process_volume.svg';
+        return 'packages/epic_common/assets/icons/post_process_volume.svg';
 
       default:
-        return 'assets/images/icons/color_grading.svg';
+        return 'packages/epic_common/assets/icons/color_grading.svg';
     }
   }
 
@@ -896,18 +889,18 @@ class _ColorGradingTargetPanelState extends State<_ColorGradingTargetPanel> {
                     padding: EdgeInsets.symmetric(horizontal: 4),
                     child: Row(children: [
                       CardSubHeaderButton(
-                        iconPath: 'assets/images/icons/add_circle.svg',
+                        iconPath: 'packages/epic_common/assets/icons/add_circle.svg',
                         tooltipMessage: AppLocalizations.of(context)!.colorGradingOutlinerAddTargetButtonTooltip,
                         onPressed: widget.objectEntry?.targetListProperty != null ? _addTargetToSelectedObject : null,
                       ),
                       CardSubHeaderButton(
-                        iconPath: 'assets/images/icons/edit.svg',
+                        iconPath: 'packages/epic_common/assets/icons/edit.svg',
                         tooltipMessage: AppLocalizations.of(context)!.colorGradingOutlinerRenameTargetButtonTooltip,
                         onPressed: widget.currentTarget?.nameProperty != null ? _renameTarget : null,
                       ),
                       const Spacer(),
                       CardSubHeaderButton(
-                        iconPath: 'assets/images/icons/trash.svg',
+                        iconPath: 'packages/epic_common/assets/icons/trash.svg',
                         tooltipMessage: AppLocalizations.of(context)!.colorGradingOutlinerDeleteTargetButtonTooltip,
                         onPressed: connection.apiVersion?.bCanUseQueryParamsInWebSocketHttpUrl == true &&
                                 widget.currentTarget?.bIsListEntry == true
@@ -1236,10 +1229,10 @@ class _ColorGradingObjectEntry extends StatelessWidget {
   }) : super(key: key);
 
   static const Map<_ColorGradingObjectEntryType, String> _iconsByType = {
-    _ColorGradingObjectEntryType.level: 'assets/images/icons/level.svg',
-    _ColorGradingObjectEntryType.nDisplayConfig: 'assets/images/icons/ndisplay.svg',
-    _ColorGradingObjectEntryType.icvfxCamera: 'assets/images/icons/ndisplay_camera.svg',
-    _ColorGradingObjectEntryType.postProcessVolume: 'assets/images/icons/post_process_volume.svg',
+    _ColorGradingObjectEntryType.level: 'packages/epic_common/assets/icons/level.svg',
+    _ColorGradingObjectEntryType.nDisplayConfig: 'packages/epic_common/assets/icons/ndisplay.svg',
+    _ColorGradingObjectEntryType.icvfxCamera: 'packages/epic_common/assets/icons/ndisplay_camera.svg',
+    _ColorGradingObjectEntryType.postProcessVolume: 'packages/epic_common/assets/icons/post_process_volume.svg',
   };
 
   /// The node in the tree view for this item.
@@ -1329,7 +1322,7 @@ class _ColorGradingTargetListEntry extends StatelessWidget {
       builder: (_, bIsEnabled) => CardListTile(
         bIsSelected: bIsSelected,
         title: data.name,
-        iconPath: 'assets/images/icons/viewport.svg',
+        iconPath: 'packages/epic_common/assets/icons/viewport.svg',
         bDeEmphasize: !bIsEnabled,
         onTap: onTap,
       ),
@@ -1358,8 +1351,8 @@ class _EnablePropertySwipeRevealer extends StatelessWidget {
       builder: (_, final bool? bIsEnabled, modify) => SwipeRevealer(
         rightSwipeActionBuilder: (context, onFinished) => CardListTileSwipeAction(
           iconPath: (bIsEnabled == true)
-              ? 'assets/images/icons/checkbox_opaque_checked.svg'
-              : 'assets/images/icons/checkbox_opaque_unchecked.svg',
+              ? 'packages/epic_common/assets/icons/checkbox_opaque_checked.svg'
+              : 'packages/epic_common/assets/icons/checkbox_opaque_unchecked.svg',
           color: UnrealColors.gray22,
           iconSize: 18,
           onPressed: () {
diff --git a/Engine/Extras/VirtualProduction/EpicStageApp/lib/widgets/screens/main/tabs/details_tab.dart b/Engine/Extras/VirtualProduction/EpicStageApp/lib/widgets/screens/main/tabs/details_tab.dart
--- a/Engine/Extras/VirtualProduction/EpicStageApp/lib/widgets/screens/main/tabs/details_tab.dart
+++ b/Engine/Extras/VirtualProduction/EpicStageApp/lib/widgets/screens/main/tabs/details_tab.dart
@@ -3,6 +3,9 @@
 import 'dart:async';
 import 'dart:collection';
 
+import 'package:epic_common/preferences.dart';
+import 'package:epic_common/theme.dart';
+import 'package:epic_common/widgets.dart';
 import 'package:flutter/material.dart';
 import 'package:flutter_gen/gen_l10n/app_localizations.dart';
 import 'package:provider/provider.dart';
@@ -17,15 +20,10 @@ import '../../../../models/unreal_actor_manager.dart';
 import '../../../../models/unreal_types.dart';
 import '../../../../utilities/constants.dart';
 import '../../../../utilities/guarded_refresh_state.dart';
-import '../../../../utilities/preferences_bundle.dart';
 import '../../../elements/delta_slider.dart';
 import '../../../elements/dropdown_button.dart';
 import '../../../elements/dropdown_text.dart';
-import '../../../elements/empty_placeholder.dart';
-import '../../../elements/epic_icon_button.dart';
-import '../../../elements/layout/card.dart';
 import '../../../elements/reset_mode_button.dart';
-import '../../../elements/selector_bar.dart';
 import '../../../elements/stepper.dart';
 import '../../../elements/unreal_property_builder.dart';
 import '../sidebar/outliner_panel.dart';
@@ -43,7 +41,7 @@ enum _DetailsActorType {
 class DetailsTab extends StatefulWidget {
   const DetailsTab({Key? key}) : super(key: key);
 
-  static const String iconPath = 'assets/images/icons/details.svg';
+  static const String iconPath = 'packages/epic_common/assets/icons/details.svg';
 
   static String getTitle(BuildContext context) => AppLocalizations.of(context)!.tabTitleDetails;
 
@@ -97,7 +95,7 @@ class _DetailsTabState extends State<DetailsTab> with GuardedRefreshState {
     final _DetailsActorType? actorType = _getActorType();
 
     return Padding(
-      padding: EdgeInsets.all(cardMargin),
+      padding: EdgeInsets.all(UnrealTheme.cardMargin),
       child: Row(
         children: [
           Expanded(
@@ -117,7 +115,7 @@ class _DetailsTabState extends State<DetailsTab> with GuardedRefreshState {
               preference: Provider.of<MainScreenSettings>(context, listen: false).bIsOutlinerPanelOpen,
               builder: (context, final bool bIsOutlinerPanelOpen) {
                 return Row(children: [
-                  if (bIsOutlinerPanelOpen) const SizedBox(width: cardMargin),
+                  if (bIsOutlinerPanelOpen) const SizedBox(width: UnrealTheme.cardMargin),
                   if (bIsOutlinerPanelOpen) OutlinerPanel(),
                 ]);
               }),
@@ -159,7 +157,7 @@ class _DetailsTabState extends State<DetailsTab> with GuardedRefreshState {
       return iconPath;
     }
 
-    return 'assets/images/icons/details.svg';
+    return 'packages/epic_common/assets/icons/details.svg';
   }
 
   /// Called when an actor we're editing has an update from the actor manager.
@@ -217,7 +215,7 @@ class _DetailsTabState extends State<DetailsTab> with GuardedRefreshState {
               ? null
               : EpicWideButton(
                   text: AppLocalizations.of(context)!.detailsTabShowOutlinerButtonLabel,
-                  iconPath: 'assets/images/icons/outliner.svg',
+                  iconPath: 'packages/epic_common/assets/icons/outliner.svg',
                   onPressed: () => mainScreenSettings.bIsOutlinerPanelOpen.setValue(true),
                 ),
         ),
diff --git a/Engine/Extras/VirtualProduction/EpicStageApp/lib/widgets/screens/main/tabs/stage_map_tab.dart b/Engine/Extras/VirtualProduction/EpicStageApp/lib/widgets/screens/main/tabs/stage_map_tab.dart
--- a/Engine/Extras/VirtualProduction/EpicStageApp/lib/widgets/screens/main/tabs/stage_map_tab.dart
+++ b/Engine/Extras/VirtualProduction/EpicStageApp/lib/widgets/screens/main/tabs/stage_map_tab.dart
@@ -4,17 +4,16 @@ import 'package:flutter/material.dart';
 import 'package:flutter_gen/gen_l10n/app_localizations.dart';
 import 'package:provider/provider.dart';
 import 'package:streaming_shared_preferences/streaming_shared_preferences.dart';
+import 'package:epic_common/theme.dart';
 
 import '../../../../models/settings/main_screen_settings.dart';
 import '../../../../models/unreal_types.dart';
-import '../../../../utilities/constants.dart';
-import '../../../../utilities/unreal_colors.dart';
 import '../../../elements/lightcard_map.dart';
 import '../sidebar/outliner_panel.dart';
 
 class StageMapTab extends StatefulWidget {
   const StageMapTab({Key? key}) : super(key: key);
-  static const String iconPath = 'assets/images/icons/world.svg';
+  static const String iconPath = 'packages/epic_common/assets/icons/world.svg';
 
   static String getTitle(BuildContext context) => AppLocalizations.of(context)!.tabTitleStageMap;
 
@@ -43,11 +42,11 @@ class _StageMapTabState extends State<StageMapTab> {
           if (bIsOutlinerPanelOpen)
             Container(
               key: const Key('OutlinerPanel'),
-              padding: EdgeInsets.all(cardMargin),
+              padding: EdgeInsets.all(UnrealTheme.cardMargin),
               decoration: BoxDecoration(
                 borderRadius: const BorderRadius.only(
-                  topLeft: Radius.circular(outerCornerRadius),
-                  bottomLeft: Radius.circular(outerCornerRadius),
+                  topLeft: Radius.circular(UnrealTheme.outerCornerRadius),
+                  bottomLeft: Radius.circular(UnrealTheme.outerCornerRadius),
                 ),
                 color: Theme.of(context).colorScheme.background,
               ),
diff --git a/Engine/Extras/VirtualProduction/EpicStageApp/lib/widgets/screens/main/tabs/web_browser_tab.dart b/Engine/Extras/VirtualProduction/EpicStageApp/lib/widgets/screens/main/tabs/web_browser_tab.dart
--- a/Engine/Extras/VirtualProduction/EpicStageApp/lib/widgets/screens/main/tabs/web_browser_tab.dart
+++ b/Engine/Extras/VirtualProduction/EpicStageApp/lib/widgets/screens/main/tabs/web_browser_tab.dart
@@ -1,5 +1,7 @@
 // Copyright Epic Games, Inc. All Rights Reserved.
 
+import 'package:epic_common/theme.dart';
+import 'package:epic_common/widgets.dart';
 import 'package:flutter/material.dart';
 import 'package:flutter_gen/gen_l10n/app_localizations.dart';
 import 'package:logging/logging.dart';
@@ -7,10 +9,7 @@ import 'package:provider/provider.dart';
 import 'package:webview_flutter/webview_flutter.dart';
 
 import '../../../../models/engine_connection.dart';
-import '../../../../utilities/constants.dart';
 import '../../../../utilities/net_utilities.dart';
-import '../../../elements/empty_placeholder.dart';
-import '../../../elements/epic_icon_button.dart';
 import '../../../elements/spinner_overlay.dart';
 
 final _log = Logger('WebBrowser');
@@ -18,7 +17,7 @@ final _log = Logger('WebBrowser');
 class WebBrowserTab extends StatefulWidget {
   const WebBrowserTab({Key? key}) : super(key: key);
 
-  static const String iconPath = 'assets/images/icons/web_browser.svg';
+  static const String iconPath = 'packages/epic_common/assets/icons/web_browser.svg';
 
   static String getTitle(BuildContext context) => AppLocalizations.of(context)!.tabTitleWebBrowser;
 
@@ -95,7 +94,7 @@ class _WebBrowserTabState extends State<WebBrowserTab> {
         crossAxisAlignment: CrossAxisAlignment.stretch,
         children: [
           _AddressBar(address: _currentUrl),
-          const SizedBox(height: sectionMargin),
+          const SizedBox(height: UnrealTheme.sectionMargin),
           Expanded(
             child: Stack(
               children: [
@@ -126,7 +125,7 @@ class _WebBrowserTabState extends State<WebBrowserTab> {
           message: _errorMessage!,
           button: EpicWideButton(
             text: AppLocalizations.of(context)!.webBrowserReconnectButtonLabel,
-            iconPath: 'assets/images/icons/refresh.svg',
+            iconPath: 'packages/epic_common/assets/icons/refresh.svg',
             onPressed: _attemptToConnect,
           ),
         ),
@@ -134,7 +133,7 @@ class _WebBrowserTabState extends State<WebBrowserTab> {
     }
 
     return Padding(
-      padding: const EdgeInsets.all(cardMargin),
+      padding: const EdgeInsets.all(UnrealTheme.cardMargin),
       child: Card(child: mainContent),
     );
   }
@@ -264,19 +263,19 @@ class _AddressBar extends StatelessWidget {
       height: 44,
       color: Theme.of(context).colorScheme.surfaceTint,
       child: Padding(
-        padding: const EdgeInsets.symmetric(horizontal: cardMargin),
+        padding: const EdgeInsets.symmetric(horizontal: UnrealTheme.cardMargin),
         child: Row(
           crossAxisAlignment: CrossAxisAlignment.center,
           mainAxisAlignment: MainAxisAlignment.center,
           children: [
             Flexible(
               child: Padding(
-                padding: const EdgeInsets.symmetric(vertical: cardMargin),
+                padding: const EdgeInsets.symmetric(vertical: UnrealTheme.cardMargin),
                 child: Container(
                   constraints: BoxConstraints(minWidth: 300, maxWidth: 500),
                   decoration: BoxDecoration(
                     color: Theme.of(context).colorScheme.background,
-                    borderRadius: BorderRadius.circular(outerCornerRadius),
+                    borderRadius: BorderRadius.circular(UnrealTheme.outerCornerRadius),
                   ),
                   child: Center(
                     child: Padding(
@@ -310,7 +309,7 @@ class _RefreshWebBrowserButton extends StatelessWidget {
   @override
   Widget build(BuildContext context) {
     return EpicIconButton(
-      iconPath: 'assets/images/icons/refresh.svg',
+      iconPath: 'packages/epic_common/assets/icons/refresh.svg',
       tooltipMessage: AppLocalizations.of(context)!.webBrowserRefreshButtonTooltip,
       onPressed: WebBrowserTab.refreshAll,
       buttonSize: const Size(40, 40),
diff --git a/Engine/Extras/VirtualProduction/EpicStageApp/lib/widgets/screens/main/toolbar/main_screen_tab_bar.dart b/Engine/Extras/VirtualProduction/EpicStageApp/lib/widgets/screens/main/toolbar/main_screen_tab_bar.dart
--- a/Engine/Extras/VirtualProduction/EpicStageApp/lib/widgets/screens/main/toolbar/main_screen_tab_bar.dart
+++ b/Engine/Extras/VirtualProduction/EpicStageApp/lib/widgets/screens/main/toolbar/main_screen_tab_bar.dart
@@ -1,9 +1,9 @@
 // Copyright Epic Games, Inc. All Rights Reserved.
 
+import 'package:epic_common/theme.dart';
+import 'package:epic_common/widgets.dart';
 import 'package:flutter/material.dart';
 
-import '../../../../utilities/unreal_colors.dart';
-import '../../../elements/asset_icon.dart';
 import '../tabs/main_screen_tabs.dart';
 
 /// The tab bar shown within the [MainScreenToolbar].
diff --git a/Engine/Extras/VirtualProduction/EpicStageApp/lib/widgets/screens/main/toolbar/main_screen_toolbar.dart b/Engine/Extras/VirtualProduction/EpicStageApp/lib/widgets/screens/main/toolbar/main_screen_toolbar.dart
--- a/Engine/Extras/VirtualProduction/EpicStageApp/lib/widgets/screens/main/toolbar/main_screen_toolbar.dart
+++ b/Engine/Extras/VirtualProduction/EpicStageApp/lib/widgets/screens/main/toolbar/main_screen_toolbar.dart
@@ -1,10 +1,10 @@
 // Copyright Epic Games, Inc. All Rights Reserved.
 
+import 'package:epic_common/widgets.dart';
 import 'package:flutter/material.dart';
 import 'package:flutter_gen/gen_l10n/app_localizations.dart';
 
 import '../../../elements/control_lock.dart';
-import '../../../elements/epic_icon_button.dart';
 import '../../../elements/place_actor_menu.dart';
 import '../../../elements/sensitivity_slider.dart';
 import '../sidebar/outliner_panel.dart';
@@ -75,7 +75,7 @@ class SettingsButton extends StatelessWidget {
   Widget build(BuildContext context) {
     return EpicIconButton(
       tooltipMessage: AppLocalizations.of(context)!.menuSettings,
-      iconPath: 'assets/images/icons/settings.svg',
+      iconPath: 'packages/epic_common/assets/icons/settings.svg',
       onPressed: SettingsDialog.show,
     );
   }
diff --git a/Engine/Extras/VirtualProduction/EpicStageApp/lib/widgets/screens/main/toolbar/settings/pages/settings_advanced_view.dart b/Engine/Extras/VirtualProduction/EpicStageApp/lib/widgets/screens/main/toolbar/settings/pages/settings_advanced_view.dart
--- a/Engine/Extras/VirtualProduction/EpicStageApp/lib/widgets/screens/main/toolbar/settings/pages/settings_advanced_view.dart
+++ b/Engine/Extras/VirtualProduction/EpicStageApp/lib/widgets/screens/main/toolbar/settings/pages/settings_advanced_view.dart
@@ -29,7 +29,7 @@ class SettingsAdvancedView extends StatelessWidget {
               preference: Provider.of<ConnectionSettings>(context).webSocketCompressionMode,
               builder: (context, final WebSocketCompressionMode compressionMode) => SettingsMenuItem(
                 title: AppLocalizations.of(context)!.settingsDialogCompressionLabel,
-                iconPath: 'assets/images/icons/compression.svg',
+                iconPath: 'packages/epic_common/assets/icons/compression.svg',
                 trailing: Text(getNameForWebSocketCompressionMode(context, compressionMode)),
                 onTap: () => Navigator.of(context).pushNamed(SettingsCompressionPicker.route),
               ),
diff --git a/Engine/Extras/VirtualProduction/EpicStageApp/lib/widgets/screens/main/toolbar/settings/pages/settings_compression_picker.dart b/Engine/Extras/VirtualProduction/EpicStageApp/lib/widgets/screens/main/toolbar/settings/pages/settings_compression_picker.dart
--- a/Engine/Extras/VirtualProduction/EpicStageApp/lib/widgets/screens/main/toolbar/settings/pages/settings_compression_picker.dart
+++ b/Engine/Extras/VirtualProduction/EpicStageApp/lib/widgets/screens/main/toolbar/settings/pages/settings_compression_picker.dart
@@ -36,8 +36,10 @@ class SettingsCompressionPicker extends StatelessWidget {
                       preference: connectionSettings.webSocketCompressionMode,
                       builder: (context, final WebSocketCompressionMode selectedMode) => SettingsMenuItem(
                         title: getNameForWebSocketCompressionMode(context, mode),
-                        iconPath: mode == WebSocketCompressionMode.none ? null : 'assets/images/icons/compression.svg',
-                        trailingIconPath: selectedMode == mode ? 'assets/images/icons/check.svg' : null,
+                        iconPath: mode == WebSocketCompressionMode.none
+                            ? null
+                            : 'packages/epic_common/assets/icons/compression.svg',
+                        trailingIconPath: selectedMode == mode ? 'packages/epic_common/assets/icons/check.svg' : null,
                         onTap: () => connectionSettings.webSocketCompressionMode.setValue(mode),
                       ),
                     ),
diff --git a/Engine/Extras/VirtualProduction/EpicStageApp/lib/widgets/screens/main/toolbar/settings/pages/settings_dialog_main.dart b/Engine/Extras/VirtualProduction/EpicStageApp/lib/widgets/screens/main/toolbar/settings/pages/settings_dialog_main.dart
--- a/Engine/Extras/VirtualProduction/EpicStageApp/lib/widgets/screens/main/toolbar/settings/pages/settings_dialog_main.dart
+++ b/Engine/Extras/VirtualProduction/EpicStageApp/lib/widgets/screens/main/toolbar/settings/pages/settings_dialog_main.dart
@@ -1,5 +1,7 @@
 // Copyright Epic Games, Inc. All Rights Reserved.
 
+import 'package:epic_common/theme.dart';
+import 'package:epic_common/widgets.dart';
 import 'package:flutter/material.dart';
 import 'package:flutter_gen/gen_l10n/app_localizations.dart';
 import 'package:provider/provider.dart';
@@ -9,8 +11,6 @@ import '../../../../../../models/engine_connection.dart';
 import '../../../../../../models/navigator_keys.dart';
 import '../../../../../../models/settings/selected_actor_settings.dart';
 import '../../../../../../models/unreal_actor_manager.dart';
-import '../../../../../../utilities/unreal_colors.dart';
-import '../../../../../elements/asset_icon.dart';
 import '../../../../eula/eula_screen.dart';
 import '../settings_generic.dart';
 import 'settings_advanced_view.dart';
@@ -44,7 +44,7 @@ class SettingsDialogMain extends StatelessWidget {
               stream: Provider.of<SelectedActorSettings>(context, listen: false).displayClusterRootPath,
               builder: (_, final AsyncSnapshot<String> rootPath) => SettingsMenuItem(
                 title: localizations.settingsDialogRootActorSettingLabel,
-                iconPath: 'assets/images/icons/ndisplay.svg',
+                iconPath: 'packages/epic_common/assets/icons/ndisplay.svg',
                 trailing: Text(actorManager.getActorAtPath(rootPath.data ?? '')?.name ?? ''),
                 onTap: () => Navigator.of(context).pushNamed(SettingsDialogRootActorPicker.route),
               ),
@@ -52,24 +52,24 @@ class SettingsDialogMain extends StatelessWidget {
           if (bIsConnected) const SettingsMenuDivider(),
           SettingsMenuItem(
             title: localizations.settingsDialogApplicationLogLabel,
-            iconPath: 'assets/images/icons/log.svg',
+            iconPath: 'packages/epic_common/assets/icons/log.svg',
             onTap: () => Navigator.of(context).pushNamed(SettingsLogList.route),
           ),
           SettingsMenuItem(
             title: localizations.settingsDialogAdvancedLabel,
-            iconPath: 'assets/images/icons/advanced.svg',
+            iconPath: 'packages/epic_common/assets/icons/advanced.svg',
             onTap: () => Navigator.of(context).pushNamed(SettingsAdvancedView.route),
           ),
           const SettingsMenuDivider(),
           SettingsMenuItem(
             title: localizations.settingsDialogApplicationHelpLabel,
-            iconPath: 'assets/images/icons/help.svg',
+            iconPath: 'packages/epic_common/assets/icons/help.svg',
             trailingIconPath: null,
             onTap: () => launchUrl(Uri.parse(_documentationUri)),
           ),
           SettingsMenuItem(
             title: localizations.settingsDialogApplicationAboutLabel,
-            iconPath: 'assets/images/icons/info.svg',
+            iconPath: 'packages/epic_common/assets/icons/info.svg',
             onTap: () => rootNavigatorKey.currentState?.pushNamed(EulaScreen.route, arguments: {'onPressed': () {}}),
           ),
           if (bIsConnected) const SettingsMenuDivider(),
@@ -112,7 +112,7 @@ class _DisconnectButton extends StatelessWidget {
             child: Center(
               child: Row(children: [
                 AssetIcon(
-                  path: 'assets/images/icons/exit.svg',
+                  path: 'packages/epic_common/assets/icons/exit.svg',
                   size: 24,
                   color: UnrealColors.warning,
                 ),
diff --git a/Engine/Extras/VirtualProduction/EpicStageApp/lib/widgets/screens/main/toolbar/settings/pages/settings_log_list.dart b/Engine/Extras/VirtualProduction/EpicStageApp/lib/widgets/screens/main/toolbar/settings/pages/settings_log_list.dart
--- a/Engine/Extras/VirtualProduction/EpicStageApp/lib/widgets/screens/main/toolbar/settings/pages/settings_log_list.dart
+++ b/Engine/Extras/VirtualProduction/EpicStageApp/lib/widgets/screens/main/toolbar/settings/pages/settings_log_list.dart
@@ -2,12 +2,12 @@
 
 import 'dart:io';
 
-import '../../../../../elements/epic_list_view.dart';
+import 'package:epic_common/logging.dart';
+import 'package:epic_common/widgets.dart';
 import 'package:flutter/material.dart';
 import 'package:flutter_gen/gen_l10n/app_localizations.dart';
 
 import '../../../../../../utilities/guarded_refresh_state.dart';
-import '../../../../../../utilities/logging.dart';
 import '../settings_generic.dart';
 import 'settings_log_view.dart';
 
@@ -73,7 +73,7 @@ class _SettingsLogListState extends State<SettingsLogList> with GuardedRefreshSt
 
               return SettingsMenuItem(
                 title: Logging.getNameForLog(file),
-                iconPath: 'assets/images/icons/log.svg',
+                iconPath: 'packages/epic_common/assets/icons/log.svg',
                 onTap: () => Navigator.of(context).pushNamed(
                   SettingsLogView.route,
                   arguments: SettingsLogViewArguments(
diff --git a/Engine/Extras/VirtualProduction/EpicStageApp/lib/widgets/screens/main/toolbar/settings/pages/settings_log_view.dart b/Engine/Extras/VirtualProduction/EpicStageApp/lib/widgets/screens/main/toolbar/settings/pages/settings_log_view.dart
--- a/Engine/Extras/VirtualProduction/EpicStageApp/lib/widgets/screens/main/toolbar/settings/pages/settings_log_view.dart
+++ b/Engine/Extras/VirtualProduction/EpicStageApp/lib/widgets/screens/main/toolbar/settings/pages/settings_log_view.dart
@@ -2,15 +2,13 @@
 
 import 'dart:io';
 
-import '../../../../../elements/epic_list_view.dart';
+import 'package:epic_common/logging.dart';
+import 'package:epic_common/widgets.dart';
 import 'package:flutter/material.dart';
 import 'package:share_plus/share_plus.dart';
 
 import '../../../../../../utilities/external_notifier.dart';
 import '../../../../../../utilities/guarded_refresh_state.dart';
-import '../../../../../../utilities/logging.dart';
-import '../../../../../elements/asset_icon.dart';
-import '../../../../../elements/epic_icon_button.dart';
 import '../settings_generic.dart';
 
 /// Arguments passed to the route when a [SettingsLogView] is pushed to the navigation stack.
@@ -42,7 +40,7 @@ class _SettingsLogViewState extends State<SettingsLogView> with GuardedRefreshSt
       title: Logging.getNameForLog(arguments.file),
       titleBarTrailing: EpicIconButton(
         key: shareButtonKey,
-        iconPath: 'assets/images/icons/share.svg',
+        iconPath: 'packages/epic_common/assets/icons/share.svg',
         color: Theme.of(context).colorScheme.primary,
         onPressed: () => _shareLog(arguments.file),
       ),
@@ -314,7 +312,9 @@ class _ScrollButtonState extends State<_ScrollButton> {
       duration: Duration(milliseconds: 150),
       child: FloatingActionButton(
         child: AssetIcon(
-          path: widget.bScrollToStart ? 'assets/images/icons/chevron_up.svg' : 'assets/images/icons/chevron_down.svg',
+          path: widget.bScrollToStart
+              ? 'packages/epic_common/assets/icons/chevron_up.svg'
+              : 'packages/epic_common/assets/icons/chevron_down.svg',
           size: 32,
         ),
         onPressed: _onPressed,
diff --git a/Engine/Extras/VirtualProduction/EpicStageApp/lib/widgets/screens/main/toolbar/settings/pages/settings_root_actor_picker.dart b/Engine/Extras/VirtualProduction/EpicStageApp/lib/widgets/screens/main/toolbar/settings/pages/settings_root_actor_picker.dart
--- a/Engine/Extras/VirtualProduction/EpicStageApp/lib/widgets/screens/main/toolbar/settings/pages/settings_root_actor_picker.dart
+++ b/Engine/Extras/VirtualProduction/EpicStageApp/lib/widgets/screens/main/toolbar/settings/pages/settings_root_actor_picker.dart
@@ -52,8 +52,9 @@ class _SettingsDialogRootActorPickerState extends State<SettingsDialogRootActorP
             for (final UnrealObject actor in _actorManager.getActorsOfClass(nDisplayRootActorClassName))
               SettingsMenuItem(
                 title: actor.name,
-                iconPath: 'assets/images/icons/ndisplay.svg',
-                trailingIconPath: rootActorPath.data == actor.path ? 'assets/images/icons/check.svg' : null,
+                iconPath: 'packages/epic_common/assets/icons/ndisplay.svg',
+                trailingIconPath:
+                    rootActorPath.data == actor.path ? 'packages/epic_common/assets/icons/check.svg' : null,
                 onTap: () => selectedActorSettings.displayClusterRootPath.setValue(actor.path),
               ),
           ],
diff --git a/Engine/Extras/VirtualProduction/EpicStageApp/lib/widgets/screens/main/toolbar/settings/settings_dialog.dart b/Engine/Extras/VirtualProduction/EpicStageApp/lib/widgets/screens/main/toolbar/settings/settings_dialog.dart
--- a/Engine/Extras/VirtualProduction/EpicStageApp/lib/widgets/screens/main/toolbar/settings/settings_dialog.dart
+++ b/Engine/Extras/VirtualProduction/EpicStageApp/lib/widgets/screens/main/toolbar/settings/settings_dialog.dart
@@ -1,9 +1,9 @@
 // Copyright Epic Games, Inc. All Rights Reserved.
 
+import 'package:epic_common/widgets.dart';
 import 'package:flutter/material.dart';
 
 import '../../../../../models/navigator_keys.dart';
-import '../../../../elements/modal.dart';
 import 'pages/settings_advanced_view.dart';
 import 'pages/settings_compression_picker.dart';
 import 'pages/settings_dialog_main.dart';
diff --git a/Engine/Extras/VirtualProduction/EpicStageApp/lib/widgets/screens/main/toolbar/settings/settings_generic.dart b/Engine/Extras/VirtualProduction/EpicStageApp/lib/widgets/screens/main/toolbar/settings/settings_generic.dart
--- a/Engine/Extras/VirtualProduction/EpicStageApp/lib/widgets/screens/main/toolbar/settings/settings_generic.dart
+++ b/Engine/Extras/VirtualProduction/EpicStageApp/lib/widgets/screens/main/toolbar/settings/settings_generic.dart
@@ -1,11 +1,10 @@
 // Copyright Epic Games, Inc. All Rights Reserved.
 
+import 'package:epic_common/theme.dart';
+import 'package:epic_common/widgets.dart';
 import 'package:flutter/material.dart';
 
 import '../../../../../models/navigator_keys.dart';
-import '../../../../../utilities/unreal_colors.dart';
-import '../../../../elements/asset_icon.dart';
-import '../../../../elements/epic_icon_button.dart';
 import 'settings_dialog.dart';
 
 /// Outer scaffold to use for all settings pages.
@@ -46,11 +45,11 @@ class SettingsPageScaffold extends StatelessWidget {
                 left: 0,
                 child: bCanPop
                     ? EpicIconButton(
-                        iconPath: 'assets/images/icons/chevron_left.svg',
+                        iconPath: 'packages/epic_common/assets/icons/chevron_left.svg',
                         onPressed: () => Navigator.of(context).maybePop(),
                       )
                     : EpicIconButton(
-                        iconPath: 'assets/images/icons/close.svg',
+                        iconPath: 'packages/epic_common/assets/icons/close.svg',
                         onPressed: () => Navigator.of(rootNavigatorKey.currentContext!).maybePop(),
                       ),
               ),
@@ -85,7 +84,7 @@ class SettingsMenuItem extends StatelessWidget {
     required this.title,
     this.iconPath,
     this.onTap,
-    this.trailingIconPath = 'assets/images/icons/chevron_right.svg',
+    this.trailingIconPath = 'packages/epic_common/assets/icons/chevron_right.svg',
     this.trailing,
   }) : super(key: key);
 
diff --git a/Engine/Extras/VirtualProduction/EpicStageApp/lib/widgets/search_bar.dart b/Engine/Extras/VirtualProduction/EpicStageApp/lib/widgets/search_bar.dart
deleted file mode 100644
diff --git a/Engine/Extras/VirtualProduction/EpicStageApp/pubspec.yaml b/Engine/Extras/VirtualProduction/EpicStageApp/pubspec.yaml
--- a/Engine/Extras/VirtualProduction/EpicStageApp/pubspec.yaml
+++ b/Engine/Extras/VirtualProduction/EpicStageApp/pubspec.yaml
@@ -31,6 +31,9 @@ dependencies:
   flutter:
     sdk: flutter
 
+  epic_common:
+    path: ../../FlutterCommon
+
   uuid: ^3.0.0
   archive: ^3.2.0
   web_socket_channel: ^2.1.0
@@ -78,29 +81,19 @@ flutter:
   uses-material-design: true
 
   assets:
-    - assets/images/decoration/
-    - assets/images/icons/
     - assets/images/map/
   
   shaders:
     - assets/shaders/scroll_view_fade.frag
 
-  # An image asset can refer to one or more resolution-specific "variants", see
-  # https://flutter.dev/assets-and-images/#resolution-aware.
-
-  # For details regarding adding assets from package dependencies, see
-  # https://flutter.dev/assets-and-images/#from-packages
-
   fonts:
-    - family: Droid Sans Mono
-      fonts:
-        - asset: assets/fonts/DroidSansMono.ttf
     - family: Inter
       fonts:
-        - asset: assets/fonts/Inter-VariableFont_slnt,wght.ttf
-  #
-  # For details regarding fonts from package dependencies,
-  # see https://flutter.dev/custom-fonts/#from-packages
+        - asset: packages/epic_common/assets/fonts/Inter-VariableFont_slnt,wght.ttf
+
+    - family: Droid Sans Mono
+      fonts:
+        - asset: packages/epic_common/assets/fonts/DroidSansMono.ttf
 
 # Configuration for automatic icon generation
 flutter_icons:
