EpicStageApp: Bumped version to 1.0.2

diff --git a/Engine/Extras/VirtualProduction/EpicStageApp/pubspec.yaml b/Engine/Extras/VirtualProduction/EpicStageApp/pubspec.yaml
--- a/Engine/Extras/VirtualProduction/EpicStageApp/pubspec.yaml
+++ b/Engine/Extras/VirtualProduction/EpicStageApp/pubspec.yaml
@@ -15,7 +15,7 @@ publish_to: 'none' # Remove this line if you wish to publish to pub.dev
 # In iOS, build-name is used as CFBundleShortVersionString while build-number used as CFBundleVersion.
 # Read more about iOS versioning at
 # https://developer.apple.com/library/archive/documentation/General/Reference/InfoPlistKeyReference/Articles/CoreFoundationKeys.html
-version: 1.0.1+1
+version: 1.0.2+1
 
 environment:
   sdk: ">=2.17.0 <3.0.0"
