[LiveLinkVCAM] Add missing copyright header

diff --git a/Engine/Extras/VirtualProduction/FlutterLiveLinkVCAM/IOS/Runner/Runner-Bridging-Header.h b/Engine/Extras/VirtualProduction/FlutterLiveLinkVCAM/IOS/Runner/Runner-Bridging-Header.h
--- a/Engine/Extras/VirtualProduction/FlutterLiveLinkVCAM/IOS/Runner/Runner-Bridging-Header.h
+++ b/Engine/Extras/VirtualProduction/FlutterLiveLinkVCAM/IOS/Runner/Runner-Bridging-Header.h
@@ -1 +1,3 @@
+// Copyright Epic Games, Inc. All Rights Reserved.
+
 #import "GeneratedPluginRegistrant.h"
