[UnrealStage] Bail out of ConnectMixin connect routine if widget is no longer mounted

diff --git a/Engine/Extras/VirtualProduction/EpicStageApp/lib/widgets/screens/connect/mixins/connect_mixin.dart b/Engine/Extras/VirtualProduction/EpicStageApp/lib/widgets/screens/connect/mixins/connect_mixin.dart
--- a/Engine/Extras/VirtualProduction/EpicStageApp/lib/widgets/screens/connect/mixins/connect_mixin.dart
+++ b/Engine/Extras/VirtualProduction/EpicStageApp/lib/widgets/screens/connect/mixins/connect_mixin.dart
@@ -57,6 +57,10 @@ mixin ConnectMixin<T extends StatefulWidget> on State<T> {
 
     final EngineConnectionResult result = await _connectionManager.connect(data);
 
+    if (!mounted) {
+      return result;
+    }
+
     if (result != EngineConnectionResult.success) {
       if (bIsConnecting == true) {
         Navigator.of(context).pop();
@@ -78,6 +82,10 @@ mixin ConnectMixin<T extends StatefulWidget> on State<T> {
 
     final Set<UnrealObject> actors = await _actorManager.getInitialActorsOfClass(nDisplayRootActorClassName);
 
+    if (!mounted) {
+      return result;
+    }
+
     final store = Provider.of<PreferencesBundle>(context, listen: false).persistent;
 
     if (actors.length == 1) {
@@ -96,6 +104,10 @@ mixin ConnectMixin<T extends StatefulWidget> on State<T> {
 
         final UnrealObject? newActor = await NDisplaySelectorDialog.showIfNotOpen(context);
 
+        if (!mounted) {
+          return result;
+        }
+
         if (newActor == null) {
           // User cancelled
           _connectionManager.disconnect();
