Add identifier to LiveLink app SDP answer so we can differentiate the app from other WebRTC peers.

diff --git a/Engine/Extras/VirtualProduction/LiveLinkVCAM/vcam/WebRTCStreamingConnection.swift b/Engine/Extras/VirtualProduction/LiveLinkVCAM/vcam/WebRTCStreamingConnection.swift
--- a/Engine/Extras/VirtualProduction/LiveLinkVCAM/vcam/WebRTCStreamingConnection.swift
+++ b/Engine/Extras/VirtualProduction/LiveLinkVCAM/vcam/WebRTCStreamingConnection.swift
@@ -463,14 +463,34 @@ extension WebRTCStreamingConnection: SignalClientDelegate {
         if self.webRTCClient!.hasPeerConnnection() {
             Log.info("Sending answer sdp")
             self.webRTCClient!.answer { (localSdp) in
-                Log.info(localSdp.sdp)
+                
+                let mungedSDP : RTCSessionDescription = self.addSessionIDToSDP(localSdp)
+                Log.info(mungedSDP.sdp)
                 self.hasLocalSdp = true
-                signalClient.send(sdp: localSdp)
+                signalClient.send(sdp: mungedSDP)
             }
         } else {
             Log.debug("WebRTC peer connection not setup yet - cannot handle sending answer.")
         }
     }
+    
+    func addSessionIDToSDP(_ inSDP: RTCSessionDescription) -> RTCSessionDescription {
+        // Munge the o= line of the SDP to add a unique identifier for LiveLink app
+        var sdpStr : String = inSDP.sdp
+        
+        let releaseVersionNumber: String? = Bundle.main.infoDictionary?["CFBundleShortVersionString"] as? String
+        let buildVersionNumber: String? = Bundle.main.infoDictionary?["CFBundleVersion"] as? String
+        
+        if let relNum = releaseVersionNumber, let buildNum = buildVersionNumber {
+            // Make a string that is like: s=LiveLink/1.3.2(130)
+            let replacementStr : String = "s=LiveLink/\(relNum)/(\(buildNum))"
+            
+            // Replace on the "s=-" which is session id line in the SDP
+            sdpStr = sdpStr.replacingOccurrences(of: "s=-", with: replacementStr)
+        }
+        
+        return RTCSessionDescription(type: inSDP.type, sdp: sdpStr)
+    }
 }
 
 extension WebRTCStreamingConnection: WebRTCClientDelegate {
