Bugfix: Stats graphs UI elements are updated even when they are hidden. This UI update is not required when the stats graphs are hidden. Fix is to skip this UI update if the graphs are hidden.

diff --git a/Engine/Extras/VirtualProduction/LiveLinkVCAM/vcam/LineGraphView.swift b/Engine/Extras/VirtualProduction/LiveLinkVCAM/vcam/LineGraphView.swift
--- a/Engine/Extras/VirtualProduction/LiveLinkVCAM/vcam/LineGraphView.swift
+++ b/Engine/Extras/VirtualProduction/LiveLinkVCAM/vcam/LineGraphView.swift
@@ -98,7 +98,9 @@ class LineGraphView: UIView {
         
         // Triggers a redraw (need to do this from main thread)
         DispatchQueue.main.async {
-            self.setNeedsDisplay();
+            if !self.isHidden && self.window != nil {
+                self.setNeedsDisplay()
+            }
         }
     }
     
diff --git a/Engine/Extras/VirtualProduction/LiveLinkVCAM/vcam/WebRTCStatsView.swift b/Engine/Extras/VirtualProduction/LiveLinkVCAM/vcam/WebRTCStatsView.swift
--- a/Engine/Extras/VirtualProduction/LiveLinkVCAM/vcam/WebRTCStatsView.swift
+++ b/Engine/Extras/VirtualProduction/LiveLinkVCAM/vcam/WebRTCStatsView.swift
@@ -18,6 +18,18 @@ class WebRTCStatsView: UIView {
     private var graphHeaderYPadding : Int = 30
     private var graphLeftMargin : Int = 100;
     private var graphRightMargin : Int = 100;
+    
+    override var isHidden: Bool {
+        set {
+            super.isHidden = newValue
+            for graphKV in statsGraphs {
+                graphKV.value.isHidden = newValue
+            }
+        }
+        get {
+            return super.isHidden
+        }
+    }
 
     // Graphs will be displayed in a grid layout, rowsxcols, where nRows and nCols are determined by frameWidth/graphWidth
     // When graphs exceed nCols on a single row then a new row will begin
@@ -55,6 +67,9 @@ class WebRTCStatsView: UIView {
         let graphFrame : CGRect = CGRect(x: graphX, y: graphY, width: self.graphWidth, height: self.graphHeight)
         let graph : LineGraphView = LineGraphView(frame: graphFrame)
 
+        // Graph inherits parents hidden property
+        graph.isHidden = self.isHidden
+
         // Add the graph the be drawn under this view
         self.addSubview(graph)
         
