Bugfix: LiveLink App had two datachannels, one it created and one that was automatically created as part of the request from the UE side. Fix is to remove manual creation of the datachannel on the LiveLink app side.

diff --git a/Engine/Extras/VirtualProduction/LiveLinkVCAM/vcam/WebRTC/Services/WebRTCClient.swift b/Engine/Extras/VirtualProduction/LiveLinkVCAM/vcam/WebRTC/Services/WebRTCClient.swift
--- a/Engine/Extras/VirtualProduction/LiveLinkVCAM/vcam/WebRTC/Services/WebRTCClient.swift
+++ b/Engine/Extras/VirtualProduction/LiveLinkVCAM/vcam/WebRTC/Services/WebRTCClient.swift
@@ -61,12 +61,6 @@ final class WebRTCClient: NSObject {
         
         self.peerConnection = pc
         
-        // Create the data channel
-        if let dataChannel = createDataChannel(peerConnection: pc) {
-            dataChannel.delegate = self
-            self.dataChannel = dataChannel
-        }
-        
         self.configureAudioSession()
         self.peerConnection!.delegate = self
     }
@@ -178,15 +172,6 @@ final class WebRTCClient: NSObject {
 // MARK: Data Channels
 extension WebRTCClient: RTCDataChannelDelegate {
     
-    private func createDataChannel(peerConnection : RTCPeerConnection) -> RTCDataChannel? {
-        let config = RTCDataChannelConfiguration()
-        guard let dataChannel = peerConnection.dataChannel(forLabel: "iOSDataChannel", configuration: config) else {
-            debugPrint("Warning: Couldn't create data channel.")
-            return nil
-        }
-        return dataChannel
-    }
-    
     func sendData(_ data: Data) {
         let buffer = RTCDataBuffer(data: data, isBinary: true)
         self.dataChannel?.sendData(buffer)
@@ -239,6 +224,12 @@ extension WebRTCClient: RTCPeerConnectionDelegate {
     
     func peerConnection(_ peerConnection: RTCPeerConnection, didOpen dataChannel: RTCDataChannel) {
         debugPrint("peerConnection did open data channel")
+        
+        // Store datachannel internally for message sending and set datachannel delegate so we can react to its events
+        dataChannel.delegate = self
+        self.dataChannel = dataChannel
+        
+        // Send the device resolution over the datachannel
         self.sendDeviceResolution()
     }
     
diff --git a/Engine/Extras/VirtualProduction/LiveLinkVCAM/vcam/WebRTCStreamingConnection.swift b/Engine/Extras/VirtualProduction/LiveLinkVCAM/vcam/WebRTCStreamingConnection.swift
--- a/Engine/Extras/VirtualProduction/LiveLinkVCAM/vcam/WebRTCStreamingConnection.swift
+++ b/Engine/Extras/VirtualProduction/LiveLinkVCAM/vcam/WebRTCStreamingConnection.swift
@@ -582,6 +582,8 @@ extension WebRTCStreamingConnection: WebRTCClientDelegate {
                 case .LatencyTest:
                     fallthrough
                 case .InitialSettings:
+                    // Do nothing with initial settings, but use it to send device resolution as this is a convenient time as we are guaranteed datachannel is working
+                    self.webRTCClient?.sendDeviceResolution()
                     fallthrough
                 case .FileExtension:
                     fallthrough
