Bugfix: ARKit tracking was misconfigured for our use case, was doing full tracking with world reconstruction, face detection, and many other computer vision features. We only care about device transform tracking so we can configure with a much simpler tracking configuration.

diff --git a/Engine/Extras/VirtualProduction/LiveLinkVCAM/vcam/VideoViewController+ARSessionDelegate.swift b/Engine/Extras/VirtualProduction/LiveLinkVCAM/vcam/VideoViewController+ARSessionDelegate.swift
--- a/Engine/Extras/VirtualProduction/LiveLinkVCAM/vcam/VideoViewController+ARSessionDelegate.swift
+++ b/Engine/Extras/VirtualProduction/LiveLinkVCAM/vcam/VideoViewController+ARSessionDelegate.swift
@@ -38,7 +38,12 @@ extension VideoViewController : ARSessionDelegate {
     }
     
     func session(_ session: ARSession, didUpdate frame: ARFrame) {
-        self.streamingConnection?.sendTransform(frame.camera.transform, atTime: Timecode.create().toTimeInterval())
+        
+        // We should only send tracking information if the tracking info is from a "normal" tracking state
+        // - otherwise we can get some crazy transforms when tracking is lost
+        if case .normal = frame.camera.trackingState {
+            self.streamingConnection?.sendTransform(frame.camera.transform, atTime: Timecode.create().toTimeInterval())
+        }
         
         // update controller
         self.sendControllerThumbstickUpdate()
diff --git a/Engine/Extras/VirtualProduction/LiveLinkVCAM/vcam/VideoViewController.swift b/Engine/Extras/VirtualProduction/LiveLinkVCAM/vcam/VideoViewController.swift
--- a/Engine/Extras/VirtualProduction/LiveLinkVCAM/vcam/VideoViewController.swift
+++ b/Engine/Extras/VirtualProduction/LiveLinkVCAM/vcam/VideoViewController.swift
@@ -115,8 +115,12 @@ class VideoViewController : BaseViewController {
         //self.demoModeBlurView.isHidden = !self.demoMode
         showReconnecting(false, animated: false)
         
-        let config = ARWorldTrackingConfiguration()
+        // Create a ARKit tracking config for purely tracking device transform (we don't care about the computer vision features)
+        let config = ARPositionalTrackingConfiguration()
         config.worldAlignment = .gravity
+        config.planeDetection = []
+        config.isLightEstimationEnabled = false
+        config.providesAudioData = false
         
         arSession = ARSession()
         arSession?.delegate = self
