Bugfix: LiveLink app should always request quality control ownership if it is lost as we want the camera operator to always have the best quality in a stream sharing scenario.

- Additionally, when quality control is requested LiveLinkApp will also request a keyframe to ensure stream of all parties is starting with a keyframe (without this encoding errors occur when stream sharing).

diff --git a/Engine/Extras/VirtualProduction/LiveLinkVCAM/vcam/WebRTC/Services/WebRTCClient.swift b/Engine/Extras/VirtualProduction/LiveLinkVCAM/vcam/WebRTC/Services/WebRTCClient.swift
--- a/Engine/Extras/VirtualProduction/LiveLinkVCAM/vcam/WebRTC/Services/WebRTCClient.swift
+++ b/Engine/Extras/VirtualProduction/LiveLinkVCAM/vcam/WebRTC/Services/WebRTCClient.swift
@@ -139,6 +139,18 @@ final class WebRTCClient: NSObject {
         }
     }
     
+    func sendRequestQualityControl() {
+        let bytes: [UInt8] = [PixelStreamingToStreamerMessage.RequestQualityControl.rawValue]
+        Log.info("Sending quality control request")
+        self.sendData(Data(bytes))
+    }
+    
+    func sendRequestKeyFrame(){
+        let bytes: [UInt8] = [PixelStreamingToStreamerMessage.IFrameRequest.rawValue]
+        Log.info("Sending keyframe request")
+        self.sendData(Data(bytes))
+    }
+    
     func sendDeviceResolution() {
         let bounds = UIScreen.main.nativeBounds
         
diff --git a/Engine/Extras/VirtualProduction/LiveLinkVCAM/vcam/WebRTCStreamingConnection.swift b/Engine/Extras/VirtualProduction/LiveLinkVCAM/vcam/WebRTCStreamingConnection.swift
--- a/Engine/Extras/VirtualProduction/LiveLinkVCAM/vcam/WebRTCStreamingConnection.swift
+++ b/Engine/Extras/VirtualProduction/LiveLinkVCAM/vcam/WebRTCStreamingConnection.swift
@@ -572,7 +572,13 @@ extension WebRTCStreamingConnection: WebRTCClientDelegate {
                         }
                     }
                 case .QualityControlOwnership:
-                    fallthrough
+                    // Log quality control ownership (are we the controller or not?)
+                    let isQualityController : Bool = (data[1] == 1) ? true : false
+                    Log.info("VCam is quality controller: \(isQualityController ? "true" : "false")")
+                    if !isQualityController {
+                        self.webRTCClient?.sendRequestQualityControl()
+                        self.webRTCClient?.sendRequestKeyFrame()
+                    }
                 case .Response:
                     fallthrough
                 case .FreezeFrame:
