Live Link VCAM v1.3.2 build 4

diff --git a/Engine/Extras/VirtualProduction/LiveLinkVCAM/vcam/WebRTC/Services/WebRTCClient.swift b/Engine/Extras/VirtualProduction/LiveLinkVCAM/vcam/WebRTC/Services/WebRTCClient.swift
--- a/Engine/Extras/VirtualProduction/LiveLinkVCAM/vcam/WebRTC/Services/WebRTCClient.swift
+++ b/Engine/Extras/VirtualProduction/LiveLinkVCAM/vcam/WebRTC/Services/WebRTCClient.swift
@@ -296,8 +296,8 @@ extension WebRTCClient {
     private func configureAudioSession() {
         self.rtcAudioSession.lockForConfiguration()
         do {
-            try self.rtcAudioSession.setCategory(AVAudioSession.Category.ambient.rawValue) /* Playback only */
-            try self.rtcAudioSession.setMode(AVAudioSession.Mode.default.rawValue) /* Mode to default mode */
+            try self.rtcAudioSession.setCategory(AVAudioSession.Category.ambient) /* Playback only */
+            try self.rtcAudioSession.setMode(AVAudioSession.Mode.default) /* Mode to default mode */
         } catch let error {
             debugPrint("Error changeing AVAudioSession category: \(error)")
         }
