Change default rendering rate of WebRTC's MetalKit video stream to be 60fps instead of 30fps.
- Add hinting of TrueMotion display refresh rates when video stream is shown on screen.
- Set MTKView preferredFramesPerSecond to 60.

diff --git a/Engine/Extras/VirtualProduction/LiveLinkVCAM/vcam/VideoViewController.swift b/Engine/Extras/VirtualProduction/LiveLinkVCAM/vcam/VideoViewController.swift
--- a/Engine/Extras/VirtualProduction/LiveLinkVCAM/vcam/VideoViewController.swift
+++ b/Engine/Extras/VirtualProduction/LiveLinkVCAM/vcam/VideoViewController.swift
@@ -19,6 +19,9 @@ class VideoViewController : BaseViewController {
     }
 
     private var displayLink: CADisplayLink?
+    
+    private var refreshRateHint : CADisplayLink?
+    private var lastTimestamp: CFTimeInterval = 0.0
 
     private var arSession : ARSession?
     
@@ -164,6 +167,9 @@ class VideoViewController : BaseViewController {
             guard let gc = note.object as? GCController else { return }
             self.controllerDisconnected(gamepad: gc)
         }
+        
+        // Todo: Could dynamically adjust this
+        self.setRefreshRateFps(fps: 60)
     }
     
     override func viewWillDisappear(_ animated: Bool) {
@@ -173,6 +179,39 @@ class VideoViewController : BaseViewController {
         let notificationCenter = NotificationCenter.default
         notificationCenter.removeObserver(self, name: .GCControllerDidConnect, object: nil)
         notificationCenter.removeObserver(self, name: .GCControllerDidDisconnect, object: nil)
+        
+        // If view is going away we will remove out refresh rate hint so battery usage can return to normal values
+        self.resetRefreshRateFps()
+    }
+    
+    @objc func refreshRateCallback(_ displayLink: CADisplayLink) {
+        // Todo: We could show the refresh rate in the UI from here?
+        //let deltaTime = displayLink.timestamp - self.lastTimestamp
+        //Log.info(String(deltaTime))
+        //let workingTime = displayLink.targetTimestamp - CACurrentMediaTime()
+        //Log.info(String(workingTime))
+        self.lastTimestamp = displayLink.timestamp
+    }
+    
+    func setRefreshRateFps(fps: Int) {
+        self.resetRefreshRateFps()
+        
+        // Attempt to force the display refresh rate to the 60-120hz range for WebRTC video streaming (it seems iOS does not auto detect the rate of received frames and adjust)
+        self.refreshRateHint = CADisplayLink(target: self, selector: #selector(refreshRateCallback))
+        self.refreshRateHint?.preferredFrameRateRange = CAFrameRateRange(minimum: Float(fps), maximum: Float(fps), preferred: Float(fps))
+        self.refreshRateHint?.add(to: .main, forMode: .common)
+        
+        for subview in self.renderView.subviews {
+            if let webrtcView = subview as? WebRTCView {
+                webrtcView.videoView.setPreferredFramerate(fps: fps)
+            }
+        }
+    }
+    
+    func resetRefreshRateFps() {
+        self.refreshRateHint?.remove(from: .main, forMode: .common)
+        self.refreshRateHint?.invalidate()
+        self.refreshRateHint = nil
     }
     
     func showReconnecting(_ visible : Bool, animated: Bool) {
diff --git a/Engine/Extras/VirtualProduction/LiveLinkVCAM/vcam/WebRTCView.swift b/Engine/Extras/VirtualProduction/LiveLinkVCAM/vcam/WebRTCView.swift
--- a/Engine/Extras/VirtualProduction/LiveLinkVCAM/vcam/WebRTCView.swift
+++ b/Engine/Extras/VirtualProduction/LiveLinkVCAM/vcam/WebRTCView.swift
@@ -7,6 +7,7 @@
 import Foundation
 import UIKit
 import WebRTC
+import MetalKit
 
 protocol WebRTCViewDelegate {
     func webRTCView(_ view : WebRTCView, didChangeVideoSize size : CGSize)
@@ -131,4 +132,13 @@ class RTCMTLVideoViewWithTouch : RTCMTLVideoView {
         super.touchesCancelled(touches, with: event)
         self.touchDelegate?.touchesCancelled(touches)
     }
+    
+    func setPreferredFramerate(fps: Int) {
+        for subview in self.subviews {
+            if let mtlView = subview as? MTKView {
+                mtlView.preferredFramesPerSecond = fps
+            }
+        }
+    }
+    
 }
