[LiveLinkVCAM] Implement pixel streaming with a custom WebRTC Flutter plugin (currently Android only)

diff --git a/Engine/Extras/VirtualProduction/FlutterLiveLinkVCAM/.editorconfig b/Engine/Extras/VirtualProduction/FlutterLiveLinkVCAM/.editorconfig
--- a/Engine/Extras/VirtualProduction/FlutterLiveLinkVCAM/.editorconfig
+++ b/Engine/Extras/VirtualProduction/FlutterLiveLinkVCAM/.editorconfig
@@ -9,3 +9,11 @@ indent_size = 2
 tab_width = 2
 trim_trailing_whitespace = true
 insert_final_newline = true
+
+[*.kt]
+max_line_length = 120
+indent_style = space
+indent_size = 4
+tab_width = 2
+trim_trailing_whitespace = true
+insert_final_newline = true
\ No newline at end of file
diff --git a/Engine/Extras/VirtualProduction/FlutterLiveLinkVCAM/android/app/proguard-android.txt b/Engine/Extras/VirtualProduction/FlutterLiveLinkVCAM/android/app/proguard-android.txt
new file mode 100644
--- /dev/null
+++ b/Engine/Extras/VirtualProduction/FlutterLiveLinkVCAM/android/app/proguard-android.txt
@@ -0,0 +1,8 @@
+## Flutter wrapper
+-keep class io.flutter.app.** { *; }
+-keep class io.flutter.plugin.**  { *; }
+-keep class io.flutter.util.**  { *; }
+-keep class io.flutter.view.**  { *; }
+-keep class io.flutter.**  { *; }
+-keep class io.flutter.plugins.**  { *; }
+-dontwarn io.flutter.embedding.**
\ No newline at end of file
diff --git a/Engine/Extras/VirtualProduction/FlutterLiveLinkVCAM/pigeons/README b/Engine/Extras/VirtualProduction/FlutterLiveLinkVCAM/pigeons/README
new file mode 100644
--- /dev/null
+++ b/Engine/Extras/VirtualProduction/FlutterLiveLinkVCAM/pigeons/README
@@ -0,0 +1,15 @@
+Dart files in this directory are not directly included in the app's source -- instead, they will be
+processed by Pigeon to generate code that facilitates communication between Flutter and each of
+the platform languages.
+
+You can modify the files in this directory to change the cross-platform API. Each file has an
+"@ConfigurePigeon" annotation which specifies where the resulting generated files will be stored
+for each target language. These generated files *should* be checked into Perforce, as they are not
+automatically generated each time the app is run.
+
+To re-generate Pigeon output, run this in your terminal from the app's root directory:
+
+  dart run pigeon --input pigeons/<filename>.dart
+
+For more information about Pigeon, see:
+https://pub.dev/packages/pigeon
\ No newline at end of file
diff --git a/Engine/Extras/VirtualProduction/FlutterLiveLinkVCAM/pubspec.yaml b/Engine/Extras/VirtualProduction/FlutterLiveLinkVCAM/pubspec.yaml
--- a/Engine/Extras/VirtualProduction/FlutterLiveLinkVCAM/pubspec.yaml
+++ b/Engine/Extras/VirtualProduction/FlutterLiveLinkVCAM/pubspec.yaml
@@ -1,33 +1,13 @@
 name: live_link_vcam
 description: A rebuild of the Live Link VCAM app for Flutter, enabling cross-platform deployment.
-# The following line prevents the package from being accidentally published to
-# pub.dev using `flutter pub publish`. This is preferred for private packages.
-publish_to: 'none' # Remove this line if you wish to publish to pub.dev
-
-# The following defines the version and build number for your application.
-# A version number is three numbers separated by dots, like 1.2.43
-# followed by an optional build number separated by a +.
-# Both the version and the builder number may be overridden in flutter
-# build by specifying --build-name and --build-number, respectively.
-# In Android, build-name is used as versionName while build-number used as versionCode.
-# Read more about Android versioning at https://developer.android.com/studio/publish/versioning
-# In iOS, build-name is used as CFBundleShortVersionString while build-number is used as CFBundleVersion.
-# Read more about iOS versioning at
-# https://developer.apple.com/library/archive/documentation/General/Reference/InfoPlistKeyReference/Articles/CoreFoundationKeys.html
-# In Windows, build-name is used as the major, minor, and patch parts
-# of the product and file versions while build-number is used as the build suffix.
+publish_to: 'none'
+
 version: 1.0.0+1
 
 environment:
   sdk: '>=3.0.6 <4.0.0'
   flutter: "^3.7.0"
 
-# Dependencies specify other packages that your package needs in order to work.
-# To automatically upgrade your package dependencies to the latest versions
-# consider running `flutter pub upgrade --major-versions`. Alternatively,
-# dependencies can be manually updated by changing the version numbers below to
-# the latest version available on pub.dev. To see which dependencies have newer
-# versions available, run `flutter pub outdated`.
 dependencies:
   flutter:
     sdk: flutter
@@ -37,12 +17,13 @@ dependencies:
 
   logging: 1.2.0
   web_socket_channel: ^2.1.0
-  flutter_webrtc: ^0.9.41
 
 dev_dependencies:
   flutter_test:
     sdk: flutter
 
+  pigeon: ^11.0.1
+
   # The "flutter_lints" package below contains a set of recommended lints to
   # encourage good coding practices. The lint set provided by the package is
   # activated in the `analysis_options.yaml` file located at the root of your
