[LiveLinkVCAM] iOS implementation of pixel streaming

diff --git a/Engine/Extras/VirtualProduction/FlutterLiveLinkVCAM/IOS/Runner/AppDelegate.swift b/Engine/Extras/VirtualProduction/FlutterLiveLinkVCAM/IOS/Runner/AppDelegate.swift
--- a/Engine/Extras/VirtualProduction/FlutterLiveLinkVCAM/IOS/Runner/AppDelegate.swift
+++ b/Engine/Extras/VirtualProduction/FlutterLiveLinkVCAM/IOS/Runner/AppDelegate.swift
@@ -1,13 +1,45 @@
+// Copyright Epic Games, Inc. All Rights Reserved.
+
 import UIKit
 import Flutter
+import WebRTC
 
 @UIApplicationMain
 @objc class AppDelegate: FlutterAppDelegate {
+  var peerConnectionApi: FlutterRtcPeerConnectionApi?
+  var videoViewControllerApi: FlutterRtcVideoViewControllerApi?
+  
   override func application(
     _ application: UIApplication,
     didFinishLaunchingWithOptions launchOptions: [UIApplication.LaunchOptionsKey: Any]?
   ) -> Bool {
+    RTCInitializeSSL()
+    
     GeneratedPluginRegistrant.register(with: self)
+    
+    // Register RtcVideoView platform view
+    var videoViewRegistrar: FlutterPluginRegistrar? = self.registrar(forPlugin: "RtcVideoView")
+    if (videoViewRegistrar == nil) {
+      fatalError("Failed to register RtcVideoView plugin")
+    }
+    
+    videoViewRegistrar?.register(
+      RtcVideoViewFactory(
+        messenger: videoViewRegistrar!.messenger()
+      ),
+      withId: "com.epicgames.live_link_vcam.view.RtcVideoView"
+    )
+    
+    // Register Pigeon APIs
+    let controller = window?.rootViewController as! FlutterViewController
+    
+    peerConnectionApi = FlutterRtcPeerConnectionApi(binaryMessenger: controller.binaryMessenger)
+    videoViewControllerApi = FlutterRtcVideoViewControllerApi(binaryMessenger: controller.binaryMessenger)
+  
     return super.application(application, didFinishLaunchingWithOptions: launchOptions)
   }
+  
+  override func applicationWillTerminate(_ application: UIApplication) {
+    peerConnectionApi?.disposeAll()
+  }
 }
diff --git a/Engine/Extras/VirtualProduction/FlutterLiveLinkVCAM/IOS/Runner/Util/FlutterPluginApi.swift b/Engine/Extras/VirtualProduction/FlutterLiveLinkVCAM/IOS/Runner/Util/FlutterPluginApi.swift
new file mode 100644
--- /dev/null
+++ b/Engine/Extras/VirtualProduction/FlutterLiveLinkVCAM/IOS/Runner/Util/FlutterPluginApi.swift
@@ -0,0 +1,21 @@
+// Copyright Epic Games, Inc. All Rights Reserved.
+
+import Foundation
+
+/// Base class for all APIs that communicate to and from a Flutter plugin.
+class FlutterPluginApi<T> {
+  /// The API used to send messages to Flutter.
+  private let flutter: T
+  
+  init (flutter: T) {
+    self.flutter = flutter
+  }
+  
+  /// Call a function on the Flutter API.
+  /// This will run on the main thread as required for any calls from platform to Flutter.
+  func callFlutter(call: @escaping (T) -> Void) {
+    DispatchQueue.main.async {
+      call(self.flutter)
+    }
+  }
+}
diff --git a/Engine/Extras/VirtualProduction/FlutterLiveLinkVCAM/IOS/Runner/Util/IdObject.swift b/Engine/Extras/VirtualProduction/FlutterLiveLinkVCAM/IOS/Runner/Util/IdObject.swift
new file mode 100644
--- /dev/null
+++ b/Engine/Extras/VirtualProduction/FlutterLiveLinkVCAM/IOS/Runner/Util/IdObject.swift
@@ -0,0 +1,64 @@
+// Copyright Epic Games, Inc. All Rights Reserved.
+
+import Foundation
+
+/// An object associated with an ID which can be used to refer to it when communicating with Flutter.
+class IdObject: NSObject {
+  /// Index that indicates an IndexedObject that is not managed.
+  static let invalidId: Int64 = -1
+  
+  /// The unique ID referring to this object within its manager.
+  private(set) var id: Int64 = invalidId
+  
+  /// The manager that tracks this object.
+  private var manager: (any BaseIdObjectManager)? = nil
+  
+  /// Whether this has been disposed of by its manager.
+  private var bIsDisposed: Bool = false
+  
+  /**
+   Called when this is registered to a manager and provided an ID.
+   
+   - Returns: false if this failed to initialize, else true
+   */
+  func onRegistered(id: Int64, manager: (any BaseIdObjectManager)) -> Bool {
+    assert(self.manager == nil, "IndexedObject initialized more than once")
+    
+    self.id = id
+    self.manager = manager
+    
+    if (!onIdReady()) {
+      self.id = IdObject.invalidId
+      self.manager = nil
+      return false
+    }
+    
+    return true
+  }
+  
+  /**
+   Called when this is unregistered from the manager.
+   */
+  func onUnregistered() {
+    if (bIsDisposed) {
+      return
+    }
+    
+    dispose()
+    bIsDisposed = true
+  }
+  
+  /**
+   Initialize any resources used by this object after this has registered with a manager.
+   
+   - Returns: false if this failed to initialize, else true
+   */
+  func onIdReady() -> Bool {
+    return true
+  }
+  
+  /**
+   Dispose of any resources held by this object.
+   */
+  func dispose() {}
+}
diff --git a/Engine/Extras/VirtualProduction/FlutterLiveLinkVCAM/IOS/Runner/Util/IdObjectManager.swift b/Engine/Extras/VirtualProduction/FlutterLiveLinkVCAM/IOS/Runner/Util/IdObjectManager.swift
new file mode 100644
--- /dev/null
+++ b/Engine/Extras/VirtualProduction/FlutterLiveLinkVCAM/IOS/Runner/Util/IdObjectManager.swift
@@ -0,0 +1,104 @@
+// Copyright Epic Games, Inc. All Rights Reserved.
+
+import Foundation
+
+// MARK: BaseIdObjectManager
+/// Generic protoccol for IdObjectManager, allowing the class to be referred to generically.
+protocol BaseIdObjectManager {
+  associatedtype IdObjectType
+  
+  /**
+   Start tracking a new object instance.
+   
+   - Parameter newObject: The object to start tracking.
+   - Returns: The object's ID for future lookup in this manager.
+   */
+  func register(newObject: IdObjectType) -> Int64
+  
+  /**
+   Dispose of an object instance created in this manager.
+   
+   - Parameter id: The ID of the object.
+   - Returns: The object that was disposed, if any.
+   */
+  func unregister(id: Int64) -> IdObjectType?
+  
+  /**
+   Dispose of all object instances in this manager.
+   */
+  func disposeAll()
+  
+  /**
+   Get an object managed by this manager.
+   
+   - Parameter id: The index of the object.
+   - Returns: The object with that ID, if any.
+   */
+  func get(id: Int64) -> IdObjectType?
+  
+  /**
+   Get an object managed by this manager and throw an exception if it doesn't exist.
+   
+   - Parameter id: The index of the object.
+   - Returns: The object with that ID, if any.
+   */
+  func getChecked(id: Int64) -> IdObjectType
+}
+
+// MARK: IdObjectManager
+/// Maintains a map of IDs to a corresponding IdObject, enabling lookup by their IDs.
+class IdObjectManager<T: IdObject> {
+  /// Name to display in debug messages.
+  private let debugTypeName: String
+  
+  /// Map from index to a managed object.
+  private var objectMap: Dictionary<Int64, T> = [:]
+  
+  /// The next index to use for the next created object.
+  private var nextIndex: Int64 = 0
+  
+  init(debugTypeName: String) {
+    self.debugTypeName = debugTypeName
+  }
+}
+
+// MARK: BaseIdObjectManager Impl
+extension IdObjectManager: BaseIdObjectManager {
+  @discardableResult
+  func register(newObject: T) -> Int64 {
+    if (!newObject.onRegistered(id: nextIndex, manager: self)) {
+      return IdObject.invalidId
+    }
+    
+    objectMap[nextIndex] = newObject
+    
+    nextIndex += 1
+    return nextIndex - 1
+  }
+  
+  @discardableResult
+  func unregister(id: Int64) -> T? {
+    let removed = objectMap.removeValue(forKey: id)
+    removed?.onUnregistered()
+    
+    return removed
+  }
+  
+  func disposeAll() {
+    while (!objectMap.isEmpty) {
+      objectMap.values.first?.onUnregistered()
+    }
+  }
+  
+  func get(id: Int64) -> T? {
+    return objectMap[id]
+  }
+  
+  func getChecked(id: Int64) -> T {
+    let foundObject = get(id: id)
+    
+    assert(foundObject != nil, "No \(debugTypeName) with ID \(id)")
+    
+    return foundObject!
+  }
+}
diff --git a/Engine/Extras/VirtualProduction/FlutterLiveLinkVCAM/IOS/Runner/Util/WrappedIdObject.swift b/Engine/Extras/VirtualProduction/FlutterLiveLinkVCAM/IOS/Runner/Util/WrappedIdObject.swift
new file mode 100644
--- /dev/null
+++ b/Engine/Extras/VirtualProduction/FlutterLiveLinkVCAM/IOS/Runner/Util/WrappedIdObject.swift
@@ -0,0 +1,12 @@
+// Copyright Epic Games, Inc. All Rights Reserved.
+
+import Foundation
+
+/// A thin wrapper for an object that allows it to be registered and looked up by its ID.
+class WrappedIdObject<T>: IdObject {
+  let inner: T
+  
+  init(inner: T) {
+    self.inner = inner
+  }
+}
diff --git a/Engine/Extras/VirtualProduction/FlutterLiveLinkVCAM/IOS/Runner/View/RtcVideoView.swift b/Engine/Extras/VirtualProduction/FlutterLiveLinkVCAM/IOS/Runner/View/RtcVideoView.swift
new file mode 100644
--- /dev/null
+++ b/Engine/Extras/VirtualProduction/FlutterLiveLinkVCAM/IOS/Runner/View/RtcVideoView.swift
@@ -0,0 +1,95 @@
+// Copyright Epic Games, Inc. All Rights Reserved.
+
+import Foundation
+import MetalKit
+import WebRTC
+
+/**
+ Native implementation of the Flutter RtcVideoView.
+ Listens to a FlutterRtcVideoViewController and displays the associated video track.
+ */
+class RtcVideoView: NSObject {
+  /// The view rendered to the screen.
+  let videoView = RTCMTLVideoView()
+  
+  /// Display link used to set the desired refresh rate.
+  private var refreshRateHint : CADisplayLink?
+  
+  /// The current track to display.
+  private var track: RTCVideoTrack?
+  
+  init(
+    frame: CGRect,
+    viewId: Int64,
+    creationParams: Dictionary<String, Any?>,
+    messenger: FlutterBinaryMessenger?
+  ) {
+    super.init()
+    
+    videoView.backgroundColor = UIColor.clear
+    videoView.videoContentMode = .scaleAspectFit
+    videoView.isEnabled = true
+    videoView.isHidden = false
+    
+    setRefreshRateFps(fps: 60)
+    
+    guard let controllerId: Int64 = creationParams["controllerId"] as? Int64 else {
+      fatalError("No/invalid controller ID provided")
+    }
+    
+    guard let controller = FlutterRtcVideoViewControllerApi
+      .instance?
+      .videoViewControllerManager
+      .get(id: controllerId) else {
+      fatalError("Controller #\(controllerId) does not exist")
+    }
+    
+    controller.track?.add(videoView)
+    controller.addListener(listener: self)
+  }
+  
+  deinit {
+    resetRefreshRateFps()
+  }
+  
+  private func setRefreshRateFps(fps: Int) {
+      self.resetRefreshRateFps()
+      
+      // Attempt to force the display refresh rate to the 60-120hz range for WebRTC video streaming (it seems iOS does not auto detect the rate of received frames and adjust)
+      self.refreshRateHint = CADisplayLink(target: self, selector: #selector(refreshRateCallback))
+      self.refreshRateHint?.preferredFrameRateRange = CAFrameRateRange(minimum: Float(fps), maximum: Float(fps), preferred: Float(fps))
+      self.refreshRateHint?.add(to: .main, forMode: .common)
+      
+      for subview in videoView.subviews {
+        if let mtlView = subview as? MTKView {
+            mtlView.preferredFramesPerSecond = fps
+        }
+      }
+  }
+  
+  @objc func refreshRateCallback(_ displayLink: CADisplayLink) {
+    videoView.renderFrame(nil)
+  }
+  
+  private func resetRefreshRateFps() {
+      self.refreshRateHint?.remove(from: .main, forMode: .common)
+      self.refreshRateHint?.invalidate()
+      self.refreshRateHint = nil
+  }
+}
+
+extension RtcVideoView: FlutterPlatformView {
+  func view() -> UIView {
+    return videoView
+  }
+}
+
+extension RtcVideoView: FlutterRtcVideoViewControllerListener {
+  func videoViewController(didChangeVideoTrack track: RTCVideoTrack?) {
+    self.track?.remove(videoView)
+    self.track = track
+    track?.add(videoView)
+    
+    videoView.renderFrame(nil)
+  }
+}
diff --git a/Engine/Extras/VirtualProduction/FlutterLiveLinkVCAM/IOS/Runner/View/RtcVideoViewFactory.swift b/Engine/Extras/VirtualProduction/FlutterLiveLinkVCAM/IOS/Runner/View/RtcVideoViewFactory.swift
new file mode 100644
--- /dev/null
+++ b/Engine/Extras/VirtualProduction/FlutterLiveLinkVCAM/IOS/Runner/View/RtcVideoViewFactory.swift
@@ -0,0 +1,27 @@
+// Copyright Epic Games, Inc. All Rights Reserved.
+
+import Foundation
+import WebRTC
+
+/// A factory that creates [RtcVideoView]s with the given creation parameters.
+class RtcVideoViewFactory: NSObject {
+  private var messenger: FlutterBinaryMessenger
+  
+  init (messenger: FlutterBinaryMessenger) {
+    self.messenger = messenger
+    super.init()
+  }
+  
+}
+
+extension RtcVideoViewFactory: FlutterPlatformViewFactory {
+  func create(withFrame frame: CGRect, viewIdentifier viewId: Int64, arguments args: Any?) -> FlutterPlatformView {
+    let creationParams = (args as? Dictionary<String, Any?>) ?? [:]
+    
+    return RtcVideoView(frame: frame, viewId: viewId, creationParams: creationParams, messenger: messenger)
+  }
+  
+  func createArgsCodec() -> FlutterMessageCodec & NSObjectProtocol {
+    return FlutterStandardMessageCodec.sharedInstance()
+  }
+}
diff --git a/Engine/Extras/VirtualProduction/FlutterLiveLinkVCAM/IOS/Runner/WebRTC/API/FlutterRtcPeerConnectionApi.swift b/Engine/Extras/VirtualProduction/FlutterLiveLinkVCAM/IOS/Runner/WebRTC/API/FlutterRtcPeerConnectionApi.swift
new file mode 100644
--- /dev/null
+++ b/Engine/Extras/VirtualProduction/FlutterLiveLinkVCAM/IOS/Runner/WebRTC/API/FlutterRtcPeerConnectionApi.swift
@@ -0,0 +1,125 @@
+// Copyright Epic Games, Inc. All Rights Reserved.
+
+import Foundation
+import WebRTC
+
+/// Handles messages about FlutterRtcPeerConnections to and from the Flutter API.
+class FlutterRtcPeerConnectionApi: FlutterPluginApi<RtcPeerConnectionFlutterApi> {
+  init(binaryMessenger: FlutterBinaryMessenger) {
+    super.init(flutter: RtcPeerConnectionFlutterApi(binaryMessenger: binaryMessenger))
+    
+    RtcPeerConnectionHostApiSetup.setUp(binaryMessenger: binaryMessenger, api: self)
+    FlutterRtcPeerConnectionApi.instance = self
+  }
+  
+  deinit {
+    if (FlutterRtcPeerConnectionApi.instance === self) {
+      FlutterRtcPeerConnectionApi.instance = nil
+    }
+  }
+  
+  /// The active singleton instance of this class.
+  private(set) static var instance: FlutterRtcPeerConnectionApi?
+  
+  /// Manager for indexed peer connections shared with Flutter.
+  private let peerConnectionManager = IdObjectManager<FlutterRtcPeerConnection>(debugTypeName: "PeerConnection")
+  
+  /// Manager for indexed media stream tracks shared with Flutter.
+  let mediaStreamTrackManager = IdObjectManager<FlutterRtcMediaStreamTrack>(debugTypeName: "MediaStreamTrack")
+  
+  /// Video decoder used for incoming video streams. Not actually used, but necessary to create an RTCPeerConnectionFactory.
+  private let videoEncoderFactory = RTCDefaultVideoEncoderFactory()
+  
+  /// Video decoder used for incoming video streams.
+  private let videoDecoderFactory = RTCDefaultVideoDecoderFactory()
+  
+  /// Factory used to create PeerConnections as requested by the client.
+  private var _peerConnectionFactory: RTCPeerConnectionFactory?
+  private var peerConnectionFactory: RTCPeerConnectionFactory {
+    get {
+      if (_peerConnectionFactory == nil) {
+        _peerConnectionFactory = makePeerConnectionFactory()
+      }
+      
+      return _peerConnectionFactory!
+    }
+    
+    set {
+      _peerConnectionFactory = newValue
+    }
+  }
+  
+  /// Field trials passed from Flutter to be used when creating connections.
+  private var fieldTrials: Dictionary<String, String>? = [:]
+  
+  /// Clean up any allocated resources.
+  func disposeAll() {
+    peerConnectionManager.disposeAll()
+    mediaStreamTrackManager.disposeAll()
+    _peerConnectionFactory = nil
+  }
+  
+  /// Make a new factory for producing peer connections with the current configuration.
+  private func makePeerConnectionFactory() -> RTCPeerConnectionFactory {
+    RTCInitFieldTrialDictionary(fieldTrials)
+    
+    return RTCPeerConnectionFactory(encoderFactory: videoEncoderFactory, decoderFactory: videoDecoderFactory)
+  }
+}
+
+// MARK: Host API
+extension FlutterRtcPeerConnectionApi: RtcPeerConnectionHostApi {
+  func setFieldTrials(fieldTrials: [String: String]?) throws {
+    self.fieldTrials = fieldTrials
+    
+    // Dispose of existing factory since the settings were changed.
+    // This will be lazily recreated.
+    _peerConnectionFactory = nil
+  }
+  
+  func create() throws -> Int64 {
+    let config = RTCConfiguration()
+    config.sdpSemantics = .unifiedPlan
+    config.continualGatheringPolicy = .gatherContinually
+    
+    return peerConnectionManager.register(newObject: FlutterRtcPeerConnection(
+      api: self,
+      peerConnectionFactory: peerConnectionFactory,
+      config: config
+    ))
+  }
+  
+  func dispose(connectionId: Int64) throws {
+    peerConnectionManager.unregister(id: connectionId)?.onUnregistered()
+  }
+  
+  func setRemoteDescription(connectionId: Int64, description: RtcSessionDescription, completion: @escaping (Result<Void, Error>) -> Void) {
+    peerConnectionManager
+      .getChecked(id: connectionId)
+      .setRemoteDescription(
+        description: description,
+        completion: completion
+      )
+  }
+  
+  func setLocalDescription(connectionId: Int64, description: RtcSessionDescription, completion: @escaping (Result<Void, Error>) -> Void) {
+    peerConnectionManager
+      .getChecked(id: connectionId)
+      .setLocalDescription(
+        description: description,
+        completion: completion
+      )
+  }
+  
+  func addRemoteCandidate(connectionId: Int64, candidate: RtcIceCandidate, completion: @escaping (Result<Void, Error>) -> Void) {
+    peerConnectionManager
+      .getChecked(id: connectionId)
+      .addRemoteCandidate(candidate: candidate, completion: completion)
+  }
+
+  func createAnswer(connectionId: Int64, completion: @escaping (Result<RtcSessionDescription, Error>) -> Void) {
+    peerConnectionManager
+      .getChecked(id: connectionId)
+      .createAnswer(completion: completion)
+  }
+}
diff --git a/Engine/Extras/VirtualProduction/FlutterLiveLinkVCAM/IOS/Runner/WebRTC/API/FlutterRtcVideoViewControllerApi.swift b/Engine/Extras/VirtualProduction/FlutterLiveLinkVCAM/IOS/Runner/WebRTC/API/FlutterRtcVideoViewControllerApi.swift
new file mode 100644
--- /dev/null
+++ b/Engine/Extras/VirtualProduction/FlutterLiveLinkVCAM/IOS/Runner/WebRTC/API/FlutterRtcVideoViewControllerApi.swift
@@ -0,0 +1,52 @@
+// Copyright Epic Games, Inc. All Rights Reserved.
+
+import Foundation
+
+/// Handles messages about [FlutterRtcPeerConnection]s to and from the Flutter API.
+class FlutterRtcVideoViewControllerApi: FlutterPluginApi<RtcVideoViewControllerFlutterApi> {
+  init(binaryMessenger: FlutterBinaryMessenger) {
+    super.init(flutter: RtcVideoViewControllerFlutterApi(binaryMessenger: binaryMessenger))
+    
+    RtcVideoViewControllerHostApiSetup.setUp(binaryMessenger: binaryMessenger, api: self)
+    FlutterRtcVideoViewControllerApi.instance = self
+  }
+  
+  deinit {
+    if (FlutterRtcPeerConnectionApi.instance === self) {
+      FlutterRtcVideoViewControllerApi.instance = nil
+    }
+  }
+  
+  /// Manager for indexed video view controllers shared with Flutter.
+  let videoViewControllerManager = IdObjectManager<FlutterRtcVideoViewController>(debugTypeName: "VideoViewController")
+  
+  /// The active singleton instance of this class.
+  private(set) static var instance: FlutterRtcVideoViewControllerApi?
+}
+
+// MARK: Host API
+extension FlutterRtcVideoViewControllerApi: RtcVideoViewControllerHostApi {
+  func create() throws -> Int64 {
+    return videoViewControllerManager.register(
+      newObject: FlutterRtcVideoViewController(api: self)
+    )
+  }
+  
+  func setTrack(controllerId: Int64, trackId: Int64) throws {
+    let track: FlutterRtcMediaStreamTrack?
+    
+    if (trackId == IdObject.invalidId) {
+      track = nil
+    } else {
+      track = FlutterRtcPeerConnectionApi.instance?.mediaStreamTrackManager.getChecked(id: trackId)
+    }
+    
+    videoViewControllerManager
+      .getChecked(id: controllerId)
+      .setTrack(newTrack: track)
+  }
+  
+  func dispose(controllerId: Int64) throws {
+    videoViewControllerManager.unregister(id: controllerId)
+  }
+}
diff --git a/Engine/Extras/VirtualProduction/FlutterLiveLinkVCAM/IOS/Runner/WebRTC/Flutter/FlutterRtcMediaStreamTrack.swift b/Engine/Extras/VirtualProduction/FlutterLiveLinkVCAM/IOS/Runner/WebRTC/Flutter/FlutterRtcMediaStreamTrack.swift
new file mode 100644
--- /dev/null
+++ b/Engine/Extras/VirtualProduction/FlutterLiveLinkVCAM/IOS/Runner/WebRTC/Flutter/FlutterRtcMediaStreamTrack.swift
@@ -0,0 +1,10 @@
+// Copyright Epic Games, Inc. All Rights Reserved.
+
+import Foundation
+import WebRTC
+
+/**
+ * A wrapped WebRTC MediaStreamTrack which can be referred to in Flutter.
+ * These tracks are generated and disposed of automatically by WebRTC, so this is just a thin wrapper adding an index.
+ */
+class FlutterRtcMediaStreamTrack: WrappedIdObject<RTCMediaStreamTrack> {}
diff --git a/Engine/Extras/VirtualProduction/FlutterLiveLinkVCAM/IOS/Runner/WebRTC/Flutter/FlutterRtcPeerConnection.swift b/Engine/Extras/VirtualProduction/FlutterLiveLinkVCAM/IOS/Runner/WebRTC/Flutter/FlutterRtcPeerConnection.swift
new file mode 100644
--- /dev/null
+++ b/Engine/Extras/VirtualProduction/FlutterLiveLinkVCAM/IOS/Runner/WebRTC/Flutter/FlutterRtcPeerConnection.swift
@@ -0,0 +1,189 @@
+// Copyright Epic Games, Inc. All Rights Reserved.
+
+import Foundation
+import WebRTC
+
+/** A wrapped PeerConnection that communicates events to the Flutter WebRTC API. */
+class FlutterRtcPeerConnection : IdObject {
+  /// The API used to send messages to Flutter.
+  private let api: FlutterRtcPeerConnectionApi
+  
+  /// The factory used to create a new peer connection.
+  private let peerConnectionFactory: RTCPeerConnectionFactory
+  
+  /// The RTC configuration to use for this connection.
+  private let config: RTCConfiguration
+  
+  /// The underlying WebRTC peer connection.
+  private var connection: RTCPeerConnection?
+  
+  /// IDs of stream tracks belonging to this connection.
+  private var trackIds: Array<Int64> = []
+  
+  init (api: FlutterRtcPeerConnectionApi, peerConnectionFactory: RTCPeerConnectionFactory, config: RTCConfiguration) {
+    self.api = api
+    self.peerConnectionFactory = peerConnectionFactory
+    self.config = config
+  }
+  
+  /**
+   * Set the description of the remote peer.
+   * - Parameter description: The session description.
+   */
+  func setRemoteDescription(description: RtcSessionDescription, completion: @escaping (Result<Void, Error>) -> Void) {
+    assert(connection != nil)
+    
+    connection!.setRemoteDescription(RTCSessionDescription(
+      type: description.type.toNative(),
+      sdp: description.sdp
+    ))
+    { error in
+      let result: Result = (error == nil) ? Result.success(()) : Result.failure(error!)
+      completion(result)
+    }
+  }
+  
+  /**
+   * Set the description of the local peer.
+   * - Parameter description: The session description.
+   */
+  func setLocalDescription(description: RtcSessionDescription, completion: @escaping (Result<Void, Error>) -> Void) {
+    assert(connection != nil)
+    
+    connection!.setLocalDescription(RTCSessionDescription(
+      type: description.type.toNative(),
+      sdp: description.sdp
+    ))
+    { error in
+      let result: Result = (error == nil) ? Result.success(()) : Result.failure(error!)
+      completion(result)
+    }
+  }
+  
+  /**
+   * Add an ICE candidate for connecting to the remote peer.
+   * - Parameter candidate: The candidate to add.
+   */
+  func addRemoteCandidate(candidate: RtcIceCandidate, completion: @escaping (Result<Void, Error>) -> Void) {
+    assert(connection != nil)
+    
+    connection!.add(RTCIceCandidate(
+      sdp: candidate.candidate,
+      sdpMLineIndex: Int32(candidate.sdpMLineIndex),
+      sdpMid: candidate.sdpMid
+    ))
+    { error in
+      let result: Result = (error == nil) ? Result.success(()) : Result.failure(error!)
+      completion(result)
+    }
+  }
+  
+  /**
+   * Create an answer to an offered WebRTC session.
+   */
+  func createAnswer(completion: @escaping (Result<RtcSessionDescription, Error>) -> Void) {
+    assert(connection != nil)
+    
+    let constraints = RTCMediaConstraints(
+      mandatoryConstraints: [
+        kRTCMediaConstraintsOfferToReceiveAudio: kRTCMediaConstraintsValueTrue,
+        kRTCMediaConstraintsOfferToReceiveVideo: kRTCMediaConstraintsValueTrue
+      ],
+      optionalConstraints: nil
+    )
+    
+    connection!.answer(for: constraints) {
+      description, error in
+      
+      if (error != nil) {
+        completion(Result.failure(error!))
+        return
+      }
+      
+      assert(description != nil)
+      
+      completion(Result.success(
+        RtcSessionDescription(
+          type: description!.type.toFlutter(),
+          sdp: description!.sdp
+        )
+      ))
+    }
+  }
+  
+  // MARK: IdObject
+  override func onIdReady() -> Bool {
+    let constraints = RTCMediaConstraints(mandatoryConstraints: nil, optionalConstraints: nil)
+    
+    guard let newConnection = peerConnectionFactory.peerConnection(
+      with: config,
+      constraints: constraints,
+      delegate: self
+    ) else {
+      return false
+    }
+    
+    connection = newConnection
+    
+    return true
+  }
+  
+  override func dispose() {
+    for trackId in trackIds {
+      api.mediaStreamTrackManager.unregister(id: trackId)
+    }
+    
+    connection?.close()
+    connection = nil
+  }
+  // /IndexedObject
+}
+
+// MARK: Peer Connection
+extension FlutterRtcPeerConnection: RTCPeerConnectionDelegate {
+  func peerConnection(_ peerConnection: RTCPeerConnection, didGenerate candidate: RTCIceCandidate) {
+    let data = RtcIceCandidate(
+      candidate: candidate.sdp,
+      sdpMid: candidate.sdpMid ?? "",
+      sdpMLineIndex: Int64.init(candidate.sdpMLineIndex)
+    )
+    
+    api.callFlutter { flutter in
+      flutter.onIceCandidate(connectionId: self.id, candidate: data) {}
+    }
+  }
+  
+  func peerConnection(_ peerConnection: RTCPeerConnection, didChange newState: RTCPeerConnectionState) {
+    api.callFlutter { flutter in
+      flutter.onStateChanged(connectionId: self.id, state: newState.toFlutter()) {}
+    }
+  }
+  
+  func peerConnection(_ peerConnection: RTCPeerConnection, didStartReceivingOn transceiver: RTCRtpTransceiver) {
+    let track: RTCMediaStreamTrack? = transceiver.receiver.track
+    if (track == nil) {
+      return
+    }
+    
+    let trackId: Int64 = api.mediaStreamTrackManager.register(newObject: FlutterRtcMediaStreamTrack(inner: track!))
+    
+    api.callFlutter { flutter in
+      flutter.onTrack(
+        connectionId: self.id,
+        event: RtcTrackEvent(
+          trackId: trackId,
+          kind: WebRtcUtils.parseTrackKind(string: track!.kind)
+        )
+      ) {}
+    }
+  }
+  
+  func peerConnection(_ peerConnection: RTCPeerConnection, didChange stateChanged: RTCSignalingState) {}
+  func peerConnection(_ peerConnection: RTCPeerConnection, didAdd stream: RTCMediaStream) {}
+  func peerConnection(_ peerConnection: RTCPeerConnection, didRemove stream: RTCMediaStream) {}
+  func peerConnectionShouldNegotiate(_ peerConnection: RTCPeerConnection) {}
+  func peerConnection(_ peerConnection: RTCPeerConnection, didChange newState: RTCIceConnectionState) {}
+  func peerConnection(_ peerConnection: RTCPeerConnection, didChange newState: RTCIceGatheringState) {}
+  func peerConnection(_ peerConnection: RTCPeerConnection, didRemove candidates: [RTCIceCandidate]) {}
+  func peerConnection(_ peerConnection: RTCPeerConnection, didOpen dataChannel: RTCDataChannel) {}
+}
diff --git a/Engine/Extras/VirtualProduction/FlutterLiveLinkVCAM/IOS/Runner/WebRTC/Flutter/FlutterRtcVideoViewController.swift b/Engine/Extras/VirtualProduction/FlutterLiveLinkVCAM/IOS/Runner/WebRTC/Flutter/FlutterRtcVideoViewController.swift
new file mode 100644
--- /dev/null
+++ b/Engine/Extras/VirtualProduction/FlutterLiveLinkVCAM/IOS/Runner/WebRTC/Flutter/FlutterRtcVideoViewController.swift
@@ -0,0 +1,92 @@
+// Copyright Epic Games, Inc. All Rights Reserved.
+
+import Foundation
+import WebRTC
+
+protocol FlutterRtcVideoViewControllerListener: AnyObject {
+  /**
+   Called when the active video track changes.
+   
+   - Parameter track: The new video track.
+   */
+  func videoViewController(didChangeVideoTrack track: RTCVideoTrack?)
+}
+
+/// Controller for a video view which communicates with Flutter.
+class FlutterRtcVideoViewController: IdObject {
+  /// The API used to send messages to Flutter.
+  private let api: FlutterRtcVideoViewControllerApi
+  
+  /// The current video track to be displayed.
+  private var _track: RTCVideoTrack?
+  
+  /// Listeners for changes to the changes in the controller's state.
+  private var listeners: [FlutterRtcVideoViewControllerListener] = []
+  
+  /// The current video track to be displayed.
+  var track: RTCVideoTrack? {
+    get {
+      return _track
+    }
+  }
+  
+  init(api: FlutterRtcVideoViewControllerApi) {
+    self.api = api
+  }
+  
+  override func dispose() {
+    _track?.remove(self)
+  }
+  
+  /**
+   Set the Flutter track to be displayed.
+   
+   - Parameter newTrack: The new track to display.
+   */
+  func setTrack(newTrack: FlutterRtcMediaStreamTrack?) {
+    var newVideoTrack: RTCVideoTrack?
+    
+    if (newTrack?.inner != nil) {
+      newVideoTrack = newTrack!.inner as? RTCVideoTrack
+      
+      if (newVideoTrack == nil) {
+        fatalError("Track \(newTrack!.id) is not a video track")
+      }
+    }
+    
+    _track = newVideoTrack!
+    _track!.add(self)
+    
+    listeners.forEach { listener in
+      listener.videoViewController(didChangeVideoTrack: _track)
+    }
+  }
+  
+  /// Listen for state changes to this controller.
+  func addListener(listener: FlutterRtcVideoViewControllerListener) {
+    listeners.append(listener)
+  }
+  
+  /// Stop listening for state changes to this controller.
+  func removeListener(listener: FlutterRtcVideoViewControllerListener) {
+    guard let index: Int = listeners.firstIndex(where: { $0 === listener }) else {
+      return
+    }
+    
+    listeners.remove(at: index)
+  }
+}
+
+extension FlutterRtcVideoViewController: RTCVideoRenderer {
+  func setSize(_ size: CGSize) {
+    api.callFlutter { flutter in
+      flutter.onFrameSizeChanged(
+        controllerId: self.id,
+        width: Int64(size.width),
+        height: Int64(size.height)
+      ) {}
+    }
+  }
+  
+  func renderFrame(_ frame: RTCVideoFrame?) {}
+}
diff --git a/Engine/Extras/VirtualProduction/FlutterLiveLinkVCAM/IOS/Runner/WebRTC/Util/WebRtcExtensions.swift b/Engine/Extras/VirtualProduction/FlutterLiveLinkVCAM/IOS/Runner/WebRTC/Util/WebRtcExtensions.swift
new file mode 100644
--- /dev/null
+++ b/Engine/Extras/VirtualProduction/FlutterLiveLinkVCAM/IOS/Runner/WebRTC/Util/WebRtcExtensions.swift
@@ -0,0 +1,48 @@
+// Copyright Epic Games, Inc. All Rights Reserved.
+
+import Foundation
+import WebRTC
+
+extension RtcSessionDescriptionType{
+  /// Convert from a Flutter session description type to the native WebRTC equivalent.
+  func toNative() -> RTCSdpType {
+    switch (self) {
+    case .answer: return .answer
+    case .offer: return .offer
+    case .pranswer: return .prAnswer
+    case .rollback: return .rollback
+    }
+  }
+}
+
+extension RTCSdpType {
+  /// Convert from the native WebRTC session description type to the Flutter equivalent.
+  func toFlutter() -> RtcSessionDescriptionType {
+    switch (self) {
+    case .answer: return .answer
+    case .offer: return .offer
+    case .prAnswer: return .pranswer
+    case .rollback: return .rollback
+      
+    @unknown default:
+      fatalError("Unknown native session description type \(self)")
+    }
+  }
+}
+
+extension RTCPeerConnectionState {
+  /// Convert from the native WebRTC connection state to the Flutter equivalent.
+  func toFlutter() -> RtcPeerConnectionState {
+    switch (self) {
+    case .new: return .newConnection
+    case .connecting: return .connecting
+    case .connected: return .connected
+    case .disconnected: return .disconnected
+    case .failed: return .failed
+    case .closed: return .closed
+      
+    @unknown default:
+      fatalError("Unknown native peer connection state \(self)")
+    }
+  }
+}
diff --git a/Engine/Extras/VirtualProduction/FlutterLiveLinkVCAM/IOS/Runner/WebRTC/Util/WebRtcUtils.swift b/Engine/Extras/VirtualProduction/FlutterLiveLinkVCAM/IOS/Runner/WebRTC/Util/WebRtcUtils.swift
new file mode 100644
--- /dev/null
+++ b/Engine/Extras/VirtualProduction/FlutterLiveLinkVCAM/IOS/Runner/WebRTC/Util/WebRtcUtils.swift
@@ -0,0 +1,20 @@
+// Copyright Epic Games, Inc. All Rights Reserved.
+
+import Foundation
+import WebRTC
+
+class WebRtcUtils {
+  /**
+   * Given the string describing a WebRTC media track kind, return the corresponding enum value.
+   *
+   * - Parameter string: The string to parse.
+   */
+  static func parseTrackKind(string: String) -> RtcMediaStreamTrackKind {
+    switch (string) {
+    case "audio": return RtcMediaStreamTrackKind.audio
+    case "video": return RtcMediaStreamTrackKind.video
+    default:
+      fatalError("Unknown media track kind \(string)")
+    }
+  }
+}
