[LiveLinkVCAM] Implement data channel API (Swift and Kotlin). Send resolution when connected.

diff --git a/Engine/Extras/VirtualProduction/FlutterLiveLinkVCAM/IOS/Runner/AppDelegate.swift b/Engine/Extras/VirtualProduction/FlutterLiveLinkVCAM/IOS/Runner/AppDelegate.swift
--- a/Engine/Extras/VirtualProduction/FlutterLiveLinkVCAM/IOS/Runner/AppDelegate.swift
+++ b/Engine/Extras/VirtualProduction/FlutterLiveLinkVCAM/IOS/Runner/AppDelegate.swift
@@ -7,6 +7,7 @@ import WebRTC
 @UIApplicationMain
 @objc class AppDelegate: FlutterAppDelegate {
   var peerConnectionApi: FlutterRtcPeerConnectionApi?
+  var dataChannelApi: FlutterRtcDataChannelApi?
   var videoViewControllerApi: FlutterRtcVideoViewControllerApi?
   
   override func application(
@@ -31,10 +32,11 @@ import WebRTC
     )
     
     // Register Pigeon APIs
-    let controller = window?.rootViewController as! FlutterViewController
+    let binaryMessenger = (window?.rootViewController as! FlutterViewController).binaryMessenger
     
-    peerConnectionApi = FlutterRtcPeerConnectionApi(binaryMessenger: controller.binaryMessenger)
-    videoViewControllerApi = FlutterRtcVideoViewControllerApi(binaryMessenger: controller.binaryMessenger)
+    peerConnectionApi = FlutterRtcPeerConnectionApi(binaryMessenger: binaryMessenger)
+    dataChannelApi = FlutterRtcDataChannelApi(binaryMessenger: binaryMessenger)
+    videoViewControllerApi = FlutterRtcVideoViewControllerApi(binaryMessenger: binaryMessenger)
   
     return super.application(application, didFinishLaunchingWithOptions: launchOptions)
   }
diff --git a/Engine/Extras/VirtualProduction/FlutterLiveLinkVCAM/IOS/Runner/WebRTC/API/FlutterRtcDataChannelApi.swift b/Engine/Extras/VirtualProduction/FlutterLiveLinkVCAM/IOS/Runner/WebRTC/API/FlutterRtcDataChannelApi.swift
new file mode 100644
--- /dev/null
+++ b/Engine/Extras/VirtualProduction/FlutterLiveLinkVCAM/IOS/Runner/WebRTC/API/FlutterRtcDataChannelApi.swift
@@ -0,0 +1,48 @@
+// Copyright Epic Games, Inc. All Rights Reserved.
+
+import Foundation
+import WebRTC
+
+/// Handles messages about FlutterRtcDataChannels to and from the Flutter API.
+class FlutterRtcDataChannelApi: FlutterPluginApi<RtcDataChannelFlutterApi> {
+  init(binaryMessenger: FlutterBinaryMessenger) {
+    super.init(flutter: RtcDataChannelFlutterApi(binaryMessenger: binaryMessenger))
+    
+    RtcDataChannelHostApiSetup.setUp(binaryMessenger: binaryMessenger, api: self)
+    FlutterRtcDataChannelApi.instance = self
+  }
+  
+  deinit {
+    if (FlutterRtcDataChannelApi.instance === self) {
+      FlutterRtcDataChannelApi.instance = nil
+    }
+  }
+  
+  /// The active singleton instance of this class.
+  private(set) static var instance: FlutterRtcDataChannelApi?
+  
+  /// Manager for indexed data channels shared with Flutter.
+  let dataChannelManager = IdObjectManager<FlutterRtcDataChannel>(debugTypeName: "DataChannel")
+  
+  /// Convenience function to register a new WebRTC-provided data channel.
+  func register(dataChannel: RTCDataChannel) -> FlutterRtcDataChannel {
+    let flutterDataChannel = FlutterRtcDataChannel(
+      api: self,
+      channel: dataChannel
+    )
+    
+    dataChannelManager.register(newObject: flutterDataChannel)
+    
+    return flutterDataChannel
+  }
+}
+
+
+// MARK: Host API
+extension FlutterRtcDataChannelApi: RtcDataChannelHostApi {
+  func sendMessage(dataChannelId: Int64, buffer: RtcDataBuffer) throws {
+    dataChannelManager
+      .getChecked(id: dataChannelId)
+      .sendMessage(buffer: buffer)
+  }
+}
diff --git a/Engine/Extras/VirtualProduction/FlutterLiveLinkVCAM/IOS/Runner/WebRTC/API/WebRtcApi.g.swift b/Engine/Extras/VirtualProduction/FlutterLiveLinkVCAM/IOS/Runner/WebRTC/API/WebRtcApi.g.swift
--- a/Engine/Extras/VirtualProduction/FlutterLiveLinkVCAM/IOS/Runner/WebRTC/API/WebRtcApi.g.swift
+++ b/Engine/Extras/VirtualProduction/FlutterLiveLinkVCAM/IOS/Runner/WebRTC/API/WebRtcApi.g.swift
@@ -67,6 +67,40 @@ enum RtcMediaStreamTrackKind: Int {
   case video = 1
 }
 
+/// State of an [RtcDataChannel].
+enum RtcDataChannelState: Int {
+  case connecting = 0
+  case open = 1
+  case closing = 2
+  case closed = 3
+}
+
+/// A buffer passed from/to WebRTC, which may contain either binary or text data.
+///
+/// Generated class from Pigeon that represents data sent in messages.
+struct RtcDataBuffer {
+  /// The binary data contained in the buffer.
+  var data: FlutterStandardTypedData
+  /// True if the buffer contains binary data; otherwise, it contains UTF-8 text.
+  var bIsBinary: Bool
+
+  static func fromList(_ list: [Any?]) -> RtcDataBuffer? {
+    let data = list[0] as! FlutterStandardTypedData
+    let bIsBinary = list[1] as! Bool
+
+    return RtcDataBuffer(
+      data: data,
+      bIsBinary: bIsBinary
+    )
+  }
+  func toList() -> [Any?] {
+    return [
+      data,
+      bIsBinary,
+    ]
+  }
+}
+
 /// Describes the configuration of one end of a peer-to-peer WebRtc connection.
 ///
 /// Generated class from Pigeon that represents data sent in messages.
@@ -408,7 +442,7 @@ class RtcPeerConnectionFlutterApi {
       completion()
     }
   }
-  /// Called when a MediaStreamTrack with the given [trackId] and [kind] is added to the PeerConnection with the given
+  /// Called when a MediaStreamTrack (described in [event]) is added to the PeerConnection with the given
   /// [connectionId].
   func onTrack(connectionId connectionIdArg: Int64, event eventArg: RtcTrackEvent, completion: @escaping () -> Void) {
     let channel = FlutterBasicMessageChannel(name: "dev.flutter.pigeon.live_link_vcam.RtcPeerConnectionFlutterApi.onTrack", binaryMessenger: binaryMessenger, codec: codec)
@@ -416,6 +450,145 @@ class RtcPeerConnectionFlutterApi {
       completion()
     }
   }
+  /// Called when a DataChannel with the given [dataChannelId] is added to the PeerConnection with the given
+  /// [connectionId].
+  func onDataChannel(connectionId connectionIdArg: Int64, dataChannelId dataChannelIdArg: Int64, completion: @escaping () -> Void) {
+    let channel = FlutterBasicMessageChannel(name: "dev.flutter.pigeon.live_link_vcam.RtcPeerConnectionFlutterApi.onDataChannel", binaryMessenger: binaryMessenger, codec: codec)
+    channel.sendMessage([connectionIdArg, dataChannelIdArg] as [Any?]) { _ in
+      completion()
+    }
+  }
+}
+private class RtcDataChannelHostApiCodecReader: FlutterStandardReader {
+  override func readValue(ofType type: UInt8) -> Any? {
+    switch type {
+      case 128:
+        return RtcDataBuffer.fromList(self.readValue() as! [Any?])
+      default:
+        return super.readValue(ofType: type)
+    }
+  }
+}
+
+private class RtcDataChannelHostApiCodecWriter: FlutterStandardWriter {
+  override func writeValue(_ value: Any) {
+    if let value = value as? RtcDataBuffer {
+      super.writeByte(128)
+      super.writeValue(value.toList())
+    } else {
+      super.writeValue(value)
+    }
+  }
+}
+
+private class RtcDataChannelHostApiCodecReaderWriter: FlutterStandardReaderWriter {
+  override func reader(with data: Data) -> FlutterStandardReader {
+    return RtcDataChannelHostApiCodecReader(data: data)
+  }
+
+  override func writer(with data: NSMutableData) -> FlutterStandardWriter {
+    return RtcDataChannelHostApiCodecWriter(data: data)
+  }
+}
+
+class RtcDataChannelHostApiCodec: FlutterStandardMessageCodec {
+  static let shared = RtcDataChannelHostApiCodec(readerWriter: RtcDataChannelHostApiCodecReaderWriter())
+}
+
+/// API that receives messages about RtcDataChannel instances in the host language.
+///
+/// Generated protocol from Pigeon that represents a handler of messages from Flutter.
+protocol RtcDataChannelHostApi {
+  /// Send a message contained in [buffer] on the data channel with the given [dataChannelId].
+  func sendMessage(dataChannelId: Int64, buffer: RtcDataBuffer) throws
+}
+
+/// Generated setup class from Pigeon to handle messages through the `binaryMessenger`.
+class RtcDataChannelHostApiSetup {
+  /// The codec used by RtcDataChannelHostApi.
+  static var codec: FlutterStandardMessageCodec { RtcDataChannelHostApiCodec.shared }
+  /// Sets up an instance of `RtcDataChannelHostApi` to handle messages through the `binaryMessenger`.
+  static func setUp(binaryMessenger: FlutterBinaryMessenger, api: RtcDataChannelHostApi?) {
+    /// Send a message contained in [buffer] on the data channel with the given [dataChannelId].
+    let sendMessageChannel = FlutterBasicMessageChannel(name: "dev.flutter.pigeon.live_link_vcam.RtcDataChannelHostApi.sendMessage", binaryMessenger: binaryMessenger, codec: codec)
+    if let api = api {
+      sendMessageChannel.setMessageHandler { message, reply in
+        let args = message as! [Any?]
+        let dataChannelIdArg = args[0] is Int64 ? args[0] as! Int64 : Int64(args[0] as! Int32)
+        let bufferArg = args[1] as! RtcDataBuffer
+        do {
+          try api.sendMessage(dataChannelId: dataChannelIdArg, buffer: bufferArg)
+          reply(wrapResult(nil))
+        } catch {
+          reply(wrapError(error))
+        }
+      }
+    } else {
+      sendMessageChannel.setMessageHandler(nil)
+    }
+  }
+}
+private class RtcDataChannelFlutterApiCodecReader: FlutterStandardReader {
+  override func readValue(ofType type: UInt8) -> Any? {
+    switch type {
+      case 128:
+        return RtcDataBuffer.fromList(self.readValue() as! [Any?])
+      default:
+        return super.readValue(ofType: type)
+    }
+  }
+}
+
+private class RtcDataChannelFlutterApiCodecWriter: FlutterStandardWriter {
+  override func writeValue(_ value: Any) {
+    if let value = value as? RtcDataBuffer {
+      super.writeByte(128)
+      super.writeValue(value.toList())
+    } else {
+      super.writeValue(value)
+    }
+  }
+}
+
+private class RtcDataChannelFlutterApiCodecReaderWriter: FlutterStandardReaderWriter {
+  override func reader(with data: Data) -> FlutterStandardReader {
+    return RtcDataChannelFlutterApiCodecReader(data: data)
+  }
+
+  override func writer(with data: NSMutableData) -> FlutterStandardWriter {
+    return RtcDataChannelFlutterApiCodecWriter(data: data)
+  }
+}
+
+class RtcDataChannelFlutterApiCodec: FlutterStandardMessageCodec {
+  static let shared = RtcDataChannelFlutterApiCodec(readerWriter: RtcDataChannelFlutterApiCodecReaderWriter())
+}
+
+/// API that receives messages about RtcDataChannel instances in Flutter.
+///
+/// Generated class from Pigeon that represents Flutter messages that can be called from Swift.
+class RtcDataChannelFlutterApi {
+  private let binaryMessenger: FlutterBinaryMessenger
+  init(binaryMessenger: FlutterBinaryMessenger){
+    self.binaryMessenger = binaryMessenger
+  }
+  var codec: FlutterStandardMessageCodec {
+    return RtcDataChannelFlutterApiCodec.shared
+  }
+  /// Called when the [state] of the DataChannel with the given [dataChannelId] changes.
+  func onStateChanged(dataChannelId dataChannelIdArg: Int64, state stateArg: RtcDataChannelState, completion: @escaping () -> Void) {
+    let channel = FlutterBasicMessageChannel(name: "dev.flutter.pigeon.live_link_vcam.RtcDataChannelFlutterApi.onStateChanged", binaryMessenger: binaryMessenger, codec: codec)
+    channel.sendMessage([dataChannelIdArg, stateArg.rawValue] as [Any?]) { _ in
+      completion()
+    }
+  }
+  /// Called when a message contained in [buffer] is received on the data channel with the given [dataChannelId].
+  func onMessage(dataChannelId dataChannelIdArg: Int64, buffer bufferArg: RtcDataBuffer, completion: @escaping () -> Void) {
+    let channel = FlutterBasicMessageChannel(name: "dev.flutter.pigeon.live_link_vcam.RtcDataChannelFlutterApi.onMessage", binaryMessenger: binaryMessenger, codec: codec)
+    channel.sendMessage([dataChannelIdArg, bufferArg] as [Any?]) { _ in
+      completion()
+    }
+  }
 }
 /// API that receives messages about VideoViewController instances in the host language.
 ///
diff --git a/Engine/Extras/VirtualProduction/FlutterLiveLinkVCAM/IOS/Runner/WebRTC/Flutter/FlutterRtcDataChannel.swift b/Engine/Extras/VirtualProduction/FlutterLiveLinkVCAM/IOS/Runner/WebRTC/Flutter/FlutterRtcDataChannel.swift
new file mode 100644
--- /dev/null
+++ b/Engine/Extras/VirtualProduction/FlutterLiveLinkVCAM/IOS/Runner/WebRTC/Flutter/FlutterRtcDataChannel.swift
@@ -0,0 +1,55 @@
+// Copyright Epic Games, Inc. All Rights Reserved.
+
+import Foundation
+import WebRTC
+
+/// A wrapped WebRTC data channel which communicates with Flutter.
+class FlutterRtcDataChannel : IdObject {
+  /// The API used to send messages to Flutter.
+  private let api: FlutterRtcDataChannelApi
+  
+  /// The underlying WebRTC data channel.
+  private let channel: RTCDataChannel
+  
+  init (api: FlutterRtcDataChannelApi, channel: RTCDataChannel) {
+    self.api = api
+    self.channel = channel
+    
+    super.init()
+    
+    channel.delegate = self
+  }
+  
+  /// Send a message on the data channel.
+  func sendMessage(buffer: RtcDataBuffer) {
+    channel.sendData(RTCDataBuffer(data: buffer.data.data, isBinary: buffer.bIsBinary))
+  }
+}
+
+// MARK: Data Channel
+extension FlutterRtcDataChannel: RTCDataChannelDelegate {
+  func dataChannelDidChangeState(_ dataChannel: RTCDataChannel) {
+    api.callFlutter { flutter in
+      flutter.onStateChanged(
+        dataChannelId: self.id,
+        state: self.channel.readyState.toFlutter()
+      ) {}
+    }
+    
+    if (self.channel.readyState == .closed) {
+      api.dataChannelManager.unregister(id: id)
+    }
+  }
+  
+  func dataChannel(_ dataChannel: RTCDataChannel, didReceiveMessageWith buffer: RTCDataBuffer) {
+    api.callFlutter { flutter in
+      flutter.onMessage(
+        dataChannelId: self.id,
+        buffer: RtcDataBuffer(
+          data: FlutterStandardTypedData(bytes: buffer.data),
+          bIsBinary: buffer.isBinary
+        )
+      ) {}
+    }
+  }
+}
diff --git a/Engine/Extras/VirtualProduction/FlutterLiveLinkVCAM/IOS/Runner/WebRTC/Flutter/FlutterRtcPeerConnection.swift b/Engine/Extras/VirtualProduction/FlutterLiveLinkVCAM/IOS/Runner/WebRTC/Flutter/FlutterRtcPeerConnection.swift
--- a/Engine/Extras/VirtualProduction/FlutterLiveLinkVCAM/IOS/Runner/WebRTC/Flutter/FlutterRtcPeerConnection.swift
+++ b/Engine/Extras/VirtualProduction/FlutterLiveLinkVCAM/IOS/Runner/WebRTC/Flutter/FlutterRtcPeerConnection.swift
@@ -178,6 +178,14 @@ extension FlutterRtcPeerConnection: RTCPeerConnectionDelegate {
     }
   }
   
+  func peerConnection(_ peerConnection: RTCPeerConnection, didOpen dataChannel: RTCDataChannel) {
+    let dataChannel: FlutterRtcDataChannel = FlutterRtcDataChannelApi.instance!.register(dataChannel: dataChannel)
+    
+    api.callFlutter { flutter in
+      flutter.onDataChannel(connectionId: self.id, dataChannelId: dataChannel.id) {}
+    }
+  }
+  
   func peerConnection(_ peerConnection: RTCPeerConnection, didChange stateChanged: RTCSignalingState) {}
   func peerConnection(_ peerConnection: RTCPeerConnection, didAdd stream: RTCMediaStream) {}
   func peerConnection(_ peerConnection: RTCPeerConnection, didRemove stream: RTCMediaStream) {}
@@ -185,5 +193,4 @@ extension FlutterRtcPeerConnection: RTCPeerConnectionDelegate {
   func peerConnection(_ peerConnection: RTCPeerConnection, didChange newState: RTCIceConnectionState) {}
   func peerConnection(_ peerConnection: RTCPeerConnection, didChange newState: RTCIceGatheringState) {}
   func peerConnection(_ peerConnection: RTCPeerConnection, didRemove candidates: [RTCIceCandidate]) {}
-  func peerConnection(_ peerConnection: RTCPeerConnection, didOpen dataChannel: RTCDataChannel) {}
 }
diff --git a/Engine/Extras/VirtualProduction/FlutterLiveLinkVCAM/IOS/Runner/WebRTC/Util/WebRtcExtensions.swift b/Engine/Extras/VirtualProduction/FlutterLiveLinkVCAM/IOS/Runner/WebRTC/Util/WebRtcExtensions.swift
--- a/Engine/Extras/VirtualProduction/FlutterLiveLinkVCAM/IOS/Runner/WebRTC/Util/WebRtcExtensions.swift
+++ b/Engine/Extras/VirtualProduction/FlutterLiveLinkVCAM/IOS/Runner/WebRTC/Util/WebRtcExtensions.swift
@@ -46,3 +46,18 @@ extension RTCPeerConnectionState {
     }
   }
 }
+
+extension RTCDataChannelState {
+  /// Convert from the native WebRTC data channel state to the Flutter equivalent.
+  func toFlutter() -> RtcDataChannelState {
+    switch (self) {
+    case .connecting: return .connecting
+    case .open: return .open
+    case .closing: return .closing
+    case .closed: return.closed
+      
+    @unknown default:
+      fatalError("Unknown native data channel state state \(self)")
+    }
+  }
+}
