[LiveLinkVCAM] Add missing iOS files

diff --git a/Engine/Extras/VirtualProduction/FlutterLiveLinkVCAM/.gitignore b/Engine/Extras/VirtualProduction/FlutterLiveLinkVCAM/.gitignore
--- a/Engine/Extras/VirtualProduction/FlutterLiveLinkVCAM/.gitignore
+++ b/Engine/Extras/VirtualProduction/FlutterLiveLinkVCAM/.gitignore
@@ -8,7 +8,6 @@
 .buildlog/
 .history
 .svn/
-migrate_working_dir/
 
 # IntelliJ related
 *.iml
@@ -23,7 +22,7 @@ migrate_working_dir/
 
 # Flutter/Dart/Pub related
 **/doc/api/
-**/ios/Flutter/.last_build_id
+**/IOS/Flutter/.last_build_id
 .dart_tool/
 .flutter-plugins
 .flutter-plugins-dependencies
@@ -32,6 +31,9 @@ migrate_working_dir/
 .pub/
 /build/
 
+# Web related
+lib/generated_plugin_registrant.dart
+
 # Symbolication related
 app.*.symbols
 
@@ -42,3 +44,75 @@ app.*.map.json
 /android/app/debug
 /android/app/profile
 /android/app/release
+
+# Android
+gradle-wrapper.jar
+/android/.gradle
+/android/captures/
+/android/gradlew
+/android/gradlew.bat
+/android/local.properties
+/android/**/GeneratedPluginRegistrant.java
+
+# Remember to never publicly share your keystore.
+# See https://flutter.dev/docs/deployment/android#reference-the-keystore-from-the-app
+/android/key.properties
+/android/**/*.keystore
+/android/**/*.jks
+
+# iOS
+/IOS/**/dgph
+/IOS/*.mode1v3
+/IOS/*.mode2v3
+/IOS/*.moved-aside
+/IOS/*.pbxuser
+/IOS/*.perspectivev3
+/IOS/**/*sync/
+/IOS/.sconsign.dblite
+/IOS/.tags*
+/IOS/**/.vagrant/
+/IOS/**/DerivedData/
+/IOS/Icon?
+/IOS/**/Pods/
+/IOS/**/.symlinks/
+/IOS/profile
+/IOS/xcuserdata
+/IOS/**/.generated/
+/IOS/Flutter/App.framework
+/IOS/Flutter/Flutter.framework
+/IOS/Flutter/Flutter.podspec
+/IOS/Flutter/Generated.xcconfig
+/IOS/Flutter/ephemeral/
+/IOS/Flutter/app.flx
+/IOS/Flutter/app.zip
+/IOS/Flutter/flutter_assets/
+/IOS/Flutter/flutter_export_environment.sh
+/IOS/ServiceDefinitions.json
+/IOS/Runner/GeneratedPluginRegistrant.*
+/IOS/.symlinks
+
+# Exceptions to above rules.
+!/IOS/default.mode1v3
+!/IOS/default.mode2v3
+!/IOS/default.pbxuser
+!/IOS/default.perspectivev3
+
+# Windows
+/windows/flutter/ephemeral/
+/windows/flutter/generated_plugin*
+
+# Visual Studio user-specific files.
+/windows/*.suo
+/windows/*.user
+/windows/*.userosscache
+/windows/*.sln.docstates
+
+# Visual Studio build-related files.
+/windows/x64/
+/windows/x86/
+
+# Visual Studio cache files
+# files ending in .cache can be ignored
+/windows/*.[Cc]ache
+# but keep track of directories ending in .cache
+/windows/!*.[Cc]ache/
\ No newline at end of file
