LiveLink VCAM v1.3.2 build 7 : Changed fingerIds use a Set so we always take the minimum index when a new touch begins.

diff --git a/Engine/Extras/VirtualProduction/LiveLinkVCAM/vcam/WebRTC/Extensions/TouchControls.swift b/Engine/Extras/VirtualProduction/LiveLinkVCAM/vcam/WebRTC/Extensions/TouchControls.swift
--- a/Engine/Extras/VirtualProduction/LiveLinkVCAM/vcam/WebRTC/Extensions/TouchControls.swift
+++ b/Engine/Extras/VirtualProduction/LiveLinkVCAM/vcam/WebRTC/Extensions/TouchControls.swift
@@ -41,7 +41,7 @@ final class TouchControls : TouchDelegate {
     
     // We need a way to give each touch a unique finger id that is persistent throughout
     // the life of that touch. So we map each touch to an id [0...10] - iPads can do 11 simulataneous touches!
-    var fingers : [Int] = [10, 9, 8, 7, 6, 5, 4, 3, 2, 1, 0]
+    var fingers : Set = [10, 9, 8, 7, 6, 5, 4, 3, 2, 1, 0]
     var fingerIds = [UITouch : Int]()
     
     var relayTouchEvents = true {
@@ -66,9 +66,12 @@ final class TouchControls : TouchDelegate {
     }
     
     func rememberTouch(_ touch : UITouch) {
-        let fingerId : Int? = self.fingers.popLast()
+        
+        // use the smallest fingerId available
+        let fingerId : Int? = self.fingers.min()
         
         if let fingerId = fingerId {
+            self.fingers.remove(fingerId)
             self.fingerIds[touch] = fingerId
         } else {
             debugPrint("Exhausted all touch identifiers - this shouldn't happen, it indicates a leak in tracking touch events.")
@@ -79,7 +82,7 @@ final class TouchControls : TouchDelegate {
     func forgetTouch(_ touch : UITouch) {
         let touchId : Int? = self.fingerIds[touch]
         if let touchId = touchId {
-            self.fingers.append(touchId)
+            self.fingers.insert(touchId)
             self.fingerIds.removeValue(forKey: touch)
         } else {
             debugPrint("Could not forget this touch because we don't have it stored.")
