[LiveLinkVCAM] Send ARKit/ARCore camera transform to streamer

diff --git a/Engine/Extras/VirtualProduction/FlutterLiveLinkVCAM/IOS/Runner/AR/API/FlutterArSessionApi.swift b/Engine/Extras/VirtualProduction/FlutterLiveLinkVCAM/IOS/Runner/AR/API/FlutterArSessionApi.swift
new file mode 100644
--- /dev/null
+++ b/Engine/Extras/VirtualProduction/FlutterLiveLinkVCAM/IOS/Runner/AR/API/FlutterArSessionApi.swift
@@ -0,0 +1,55 @@
+// Copyright Epic Games, Inc. All Rights Reserved.
+
+import ARKit
+import Foundation
+
+/// Handles messages about FlutterArSessions to and from the Flutter API.
+class FlutterArSessionApi: FlutterPluginApi<ArSessionFlutterApi> {
+  init(binaryMessenger: FlutterBinaryMessenger) {
+    super.init(flutter: ArSessionFlutterApi(binaryMessenger: binaryMessenger))
+    
+    ArSessionHostApiSetup.setUp(binaryMessenger: binaryMessenger, api: self)
+    FlutterArSessionApi.instance = self
+  }
+  
+  deinit {
+    if (FlutterArSessionApi.instance === self) {
+      FlutterArSessionApi.instance = nil
+    }
+  }
+  
+  /// The active singleton instance of this class.
+  private(set) static var instance: FlutterArSessionApi?
+  
+  /// Manager for indexed AR sessions shared with Flutter.
+  private let sessionManager = IdObjectManager<FlutterArSession>(debugTypeName: "ARSession")
+}
+
+
+// MARK: Host API
+extension FlutterArSessionApi: ArSessionHostApi {
+  func initialize(completion: @escaping (Result<ArAvailability, Error>) -> Void) {
+    completion(Result.success(ARPositionalTrackingConfiguration.isSupported ? ArAvailability.available : ArAvailability.notSupported))
+  }
+  
+  func create(completion: @escaping (Result<Int64, Error>) -> Void) {
+    let id: Int64 = sessionManager.register(newObject: FlutterArSession(
+      api: self
+    ))
+    
+    completion(Result.success(id))
+  }
+  
+  func dispose(sessionId: Int64) throws {
+    sessionManager.unregister(id: sessionId)
+  }
+  
+  func run(sessionId: Int64) throws {
+    sessionManager.getChecked(id: sessionId).run()
+  }
+  
+  func pause(sessionId: Int64) throws {
+    sessionManager.getChecked(id: sessionId).pause()
+  }
+  
+}
diff --git a/Engine/Extras/VirtualProduction/FlutterLiveLinkVCAM/IOS/Runner/AR/Flutter/FlutterArSession.swift b/Engine/Extras/VirtualProduction/FlutterLiveLinkVCAM/IOS/Runner/AR/Flutter/FlutterArSession.swift
new file mode 100644
--- /dev/null
+++ b/Engine/Extras/VirtualProduction/FlutterLiveLinkVCAM/IOS/Runner/AR/Flutter/FlutterArSession.swift
@@ -0,0 +1,104 @@
+// Copyright Epic Games, Inc. All Rights Reserved.
+
+import ARKit
+import Foundation
+
+/** A wrapped ARSession that communicates events to the Flutter AR API. */
+class FlutterArSession: IdObject {
+  /// The API used to send messages to Flutter.
+  private let api: FlutterArSessionApi
+  
+  /// The underlying AR session.
+  private var session: ARSession?
+  
+  /// The current state of the session's camera tracking.
+  private var trackingState: ArTrackingState = ArTrackingState.unavailable
+  
+  init (api: FlutterArSessionApi) {
+    self.api = api
+    session = ARSession()
+    
+    super.init()
+    
+    session!.delegate = self
+  }
+  
+  /// Run the AR session.
+  func run() {
+    // Create an ARKit tracking config for purely tracking device transform (we don't care about the computer vision features)
+    let config = ARPositionalTrackingConfiguration()
+    config.worldAlignment = .gravity
+    config.planeDetection = []
+    config.isLightEstimationEnabled = false
+    config.providesAudioData = false
+    
+    session!.run(config)
+  }
+  
+  /// Pause the AR session.
+  func pause() {
+    session!.pause()
+  }
+  
+  override func dispose() {
+    session?.pause()
+    session?.delegate = nil
+    session = nil
+  }
+  
+  /**
+   Update the tracking state based on the state of an AR camera.
+   - Parameter newState: The new tracking state.
+   */
+  private func updateTrackingState(newState: ARCamera.TrackingState) {
+    var newTrackingState: ArTrackingState
+    
+    switch (newState) {
+    case .normal:
+      newTrackingState = .normal
+      break
+      
+    case .limited:
+      newTrackingState = .limited
+      break
+      
+    default:
+      newTrackingState = .unavailable
+      break
+    }
+    
+    if (newTrackingState != trackingState) {
+      trackingState = newTrackingState
+      
+      api.callFlutter { flutter in
+        flutter.onTrackingStateChanged(sessionId: self.id, state: self.trackingState) {}
+      }
+    }
+  }
+}
+
+extension FlutterArSession: ARSessionDelegate {
+  func session(_ session: ARSession, cameraDidChangeTrackingState camera: ARCamera) {
+    updateTrackingState(newState: camera.trackingState)
+  }
+  
+  func session(_ session: ARSession, didUpdate frame: ARFrame) {
+    updateTrackingState(newState: frame.camera.trackingState)
+    
+    if (self.trackingState != .normal) {
+      // Tracking data is inaccurate when not in normal mode
+      return
+    }
+    
+    Swift.withUnsafeBytes(of: frame.camera.transform) { matrixBytes in
+      assert(matrixBytes.baseAddress != nil)
+      
+      let data = Data(bytes: matrixBytes.baseAddress!, count: matrixBytes.count)
+      let flutterFrame = ArFrame(cameraTransformData: FlutterStandardTypedData(bytes: data))
+      
+      api.callFlutter { flutter in
+        flutter.onFrame(sessionId: self.id, frame: flutterFrame) {}
+      }
+    }
+  }
+}
diff --git a/Engine/Extras/VirtualProduction/FlutterLiveLinkVCAM/IOS/Runner/AppDelegate.swift b/Engine/Extras/VirtualProduction/FlutterLiveLinkVCAM/IOS/Runner/AppDelegate.swift
--- a/Engine/Extras/VirtualProduction/FlutterLiveLinkVCAM/IOS/Runner/AppDelegate.swift
+++ b/Engine/Extras/VirtualProduction/FlutterLiveLinkVCAM/IOS/Runner/AppDelegate.swift
@@ -9,6 +9,7 @@ import WebRTC
   var peerConnectionApi: FlutterRtcPeerConnectionApi?
   var dataChannelApi: FlutterRtcDataChannelApi?
   var videoViewControllerApi: FlutterRtcVideoViewControllerApi?
+  var arSessionApi: FlutterArSessionApi?
   
   override func application(
     _ application: UIApplication,
@@ -37,6 +38,7 @@ import WebRTC
     peerConnectionApi = FlutterRtcPeerConnectionApi(binaryMessenger: binaryMessenger)
     dataChannelApi = FlutterRtcDataChannelApi(binaryMessenger: binaryMessenger)
     videoViewControllerApi = FlutterRtcVideoViewControllerApi(binaryMessenger: binaryMessenger)
+    arSessionApi = FlutterArSessionApi(binaryMessenger: binaryMessenger)
   
     return super.application(application, didFinishLaunchingWithOptions: launchOptions)
   }
diff --git a/Engine/Extras/VirtualProduction/FlutterLiveLinkVCAM/android/app/src/main/AndroidManifest.xml b/Engine/Extras/VirtualProduction/FlutterLiveLinkVCAM/android/app/src/main/AndroidManifest.xml
--- a/Engine/Extras/VirtualProduction/FlutterLiveLinkVCAM/android/app/src/main/AndroidManifest.xml
+++ b/Engine/Extras/VirtualProduction/FlutterLiveLinkVCAM/android/app/src/main/AndroidManifest.xml
@@ -1,6 +1,8 @@
 <manifest xmlns:android="http://schemas.android.com/apk/res/android">
+
     <uses-feature android:name="android.hardware.camera" />
-    <uses-feature android:name="android.hardware.camera.autofocus" />
+    <uses-feature android:name="android.hardware.camera.ar" />
+
     <uses-permission android:name="android.permission.CAMERA" />
     <uses-permission android:name="android.permission.RECORD_AUDIO" />
     <uses-permission android:name="android.permission.INTERNET" />
@@ -33,6 +35,11 @@
                 <category android:name="android.intent.category.LAUNCHER"/>
             </intent-filter>
         </activity>
+
+        <meta-data
+            android:name="com.google.ar.core"
+            android:value="required" />
+
         <!-- Don't delete the meta-data below.
              This is used by the Flutter tool to generate GeneratedPluginRegistrant.java -->
         <meta-data
