Fix: Allow port to be configured on LiveLink VCam app.

Port is configured by entering host:port, like 127.0.0.1:8080
This fix is required on Linux, which uses port 8080 by default.

diff --git a/Engine/Extras/VirtualProduction/LiveLinkVCAM/vcam/WebRTCStreamingConnection.swift b/Engine/Extras/VirtualProduction/LiveLinkVCAM/vcam/WebRTCStreamingConnection.swift
--- a/Engine/Extras/VirtualProduction/LiveLinkVCAM/vcam/WebRTCStreamingConnection.swift
+++ b/Engine/Extras/VirtualProduction/LiveLinkVCAM/vcam/WebRTCStreamingConnection.swift
@@ -44,7 +44,10 @@ class WebRTCStreamingConnection : StreamingConnection {
             self._url?.absoluteString ?? ""
         }
         set {
-            self._url = URL(string: "ws://\(newValue):80")
+            let host : String
+            let port : UInt16?
+            (host, port) = NetUtility.hostAndPortFromAddress(newValue)
+            self._url = URL(string: "ws://\(host):\(port ?? 80)")
         }
     }
     
