Add support for android signed builds.

diff --git a/Engine/Extras/VirtualProduction/FlutterLiveLinkVCAM/README.md b/Engine/Extras/VirtualProduction/FlutterLiveLinkVCAM/README.md
--- a/Engine/Extras/VirtualProduction/FlutterLiveLinkVCAM/README.md
+++ b/Engine/Extras/VirtualProduction/FlutterLiveLinkVCAM/README.md
@@ -14,3 +14,13 @@ A few resources to get you started if this is your first Flutter project:
 For help getting started with Flutter development, view the
 [online documentation](https://docs.flutter.dev/), which offers tutorials,
 samples, guidance on mobile development, and a full API reference.
+
+## Signing Certificates 
+
+To setup signing certificates, follow the instructions outlined in this article 
+
+ - [Deploying Flutter Apps to The PlayStore](https:[REDACTED_EMAIL]/deploying-flutter-apps-to-the-playstore-1bd0cce0d15c)
+
+Add your `android-key.tks` to `android\app` folder and `key.properities` to `android\`.  The gradle files are already setup
+to read these files when app bundling
+
diff --git a/Engine/Extras/VirtualProduction/FlutterLiveLinkVCAM/android/app/src/main/AndroidManifest.xml b/Engine/Extras/VirtualProduction/FlutterLiveLinkVCAM/android/app/src/main/AndroidManifest.xml
--- a/Engine/Extras/VirtualProduction/FlutterLiveLinkVCAM/android/app/src/main/AndroidManifest.xml
+++ b/Engine/Extras/VirtualProduction/FlutterLiveLinkVCAM/android/app/src/main/AndroidManifest.xml
@@ -1,4 +1,5 @@
-<manifest xmlns:android="http://schemas.android.com/apk/res/android">
+<manifest xmlns:android="http://schemas.android.com/apk/res/android"
+	package="com.epicgames.live_link_vcam">
 
     <uses-feature android:name="android.hardware.camera" />
     <uses-feature android:name="android.hardware.camera.ar" />
