[LiveLinkVCAM] Collect and display stream stats

diff --git a/Engine/Extras/VirtualProduction/FlutterLiveLinkVCAM/IOS/Runner/WebRTC/API/FlutterRtcPeerConnectionApi.swift b/Engine/Extras/VirtualProduction/FlutterLiveLinkVCAM/IOS/Runner/WebRTC/API/FlutterRtcPeerConnectionApi.swift
--- a/Engine/Extras/VirtualProduction/FlutterLiveLinkVCAM/IOS/Runner/WebRTC/API/FlutterRtcPeerConnectionApi.swift
+++ b/Engine/Extras/VirtualProduction/FlutterLiveLinkVCAM/IOS/Runner/WebRTC/API/FlutterRtcPeerConnectionApi.swift
@@ -122,4 +122,10 @@ extension FlutterRtcPeerConnectionApi: RtcPeerConnectionHostApi {
       .getChecked(id: connectionId)
       .createAnswer(completion: completion)
   }
+  
+  func getStats(connectionId: Int64, typeFilter: [String]?, completion: @escaping (Result<RtcStatsReport, Error>) -> Void) {
+    peerConnectionManager
+      .getChecked(id: connectionId)
+      .getStats(typeFilter: typeFilter, completion: completion)
+  }
 }
diff --git a/Engine/Extras/VirtualProduction/FlutterLiveLinkVCAM/IOS/Runner/WebRTC/API/WebRtcApi.g.swift b/Engine/Extras/VirtualProduction/FlutterLiveLinkVCAM/IOS/Runner/WebRTC/API/WebRtcApi.g.swift
--- a/Engine/Extras/VirtualProduction/FlutterLiveLinkVCAM/IOS/Runner/WebRTC/API/WebRtcApi.g.swift
+++ b/Engine/Extras/VirtualProduction/FlutterLiveLinkVCAM/IOS/Runner/WebRTC/API/WebRtcApi.g.swift
@@ -184,6 +184,68 @@ struct RtcTrackEvent {
   }
 }
 
+/// Reported statistics about a single object inspected as part of a WebRTC report.
+///
+/// Generated class from Pigeon that represents data sent in messages.
+struct RtcStats {
+  /// Timestamp at which the stats were gathered (in microseconds).
+  var timestampUs: Double
+  /// The type of object for which these stats were collected.
+  var type: String
+  /// The unique ID of the stats object.
+  var id: String
+  /// Map from stat name to its value.
+  var values: [String?: Any?]
+
+  static func fromList(_ list: [Any?]) -> RtcStats? {
+    let timestampUs = list[0] as! Double
+    let type = list[1] as! String
+    let id = list[2] as! String
+    let values = list[3] as! [String?: Any?]
+
+    return RtcStats(
+      timestampUs: timestampUs,
+      type: type,
+      id: id,
+      values: values
+    )
+  }
+  func toList() -> [Any?] {
+    return [
+      timestampUs,
+      type,
+      id,
+      values,
+    ]
+  }
+}
+
+/// A collection of reported statistics gathered about a WebRTC peer connection.
+///
+/// Generated class from Pigeon that represents data sent in messages.
+struct RtcStatsReport {
+  /// Timestamp at which the stats were gathered (in microseconds).
+  var timestampUs: Double
+  /// Map from object name to the stats associated with it.
+  var stats: [String?: RtcStats?]
+
+  static func fromList(_ list: [Any?]) -> RtcStatsReport? {
+    let timestampUs = list[0] as! Double
+    let stats = list[1] as! [String?: RtcStats?]
+
+    return RtcStatsReport(
+      timestampUs: timestampUs,
+      stats: stats
+    )
+  }
+  func toList() -> [Any?] {
+    return [
+      timestampUs,
+      stats,
+    ]
+  }
+}
+
 private class RtcPeerConnectionHostApiCodecReader: FlutterStandardReader {
   override func readValue(ofType type: UInt8) -> Any? {
     switch type {
@@ -191,6 +253,10 @@ private class RtcPeerConnectionHostApiCodecReader: FlutterStandardReader {
         return RtcIceCandidate.fromList(self.readValue() as! [Any?])
       case 129:
         return RtcSessionDescription.fromList(self.readValue() as! [Any?])
+      case 130:
+        return RtcStats.fromList(self.readValue() as! [Any?])
+      case 131:
+        return RtcStatsReport.fromList(self.readValue() as! [Any?])
       default:
         return super.readValue(ofType: type)
     }
@@ -205,6 +271,12 @@ private class RtcPeerConnectionHostApiCodecWriter: FlutterStandardWriter {
     } else if let value = value as? RtcSessionDescription {
       super.writeByte(129)
       super.writeValue(value.toList())
+    } else if let value = value as? RtcStats {
+      super.writeByte(130)
+      super.writeValue(value.toList())
+    } else if let value = value as? RtcStatsReport {
+      super.writeByte(131)
+      super.writeValue(value.toList())
     } else {
       super.writeValue(value)
     }
@@ -244,6 +316,10 @@ protocol RtcPeerConnectionHostApi {
   func addRemoteCandidate(connectionId: Int64, candidate: RtcIceCandidate, completion: @escaping (Result<Void, Error>) -> Void)
   /// Create an answer to a WebRTC offer for the PeerConnection with the given [connectionId].
   func createAnswer(connectionId: Int64, completion: @escaping (Result<RtcSessionDescription, Error>) -> Void)
+  /// Generate a stats report about the PeerConnection with the given [connectionId].
+  /// If [typeFilter] is provided, the report will only contain stats matching types in that list, reducing the codec
+  /// overhead from the native platform.
+  func getStats(connectionId: Int64, typeFilter: [String]?, completion: @escaping (Result<RtcStatsReport, Error>) -> Void)
 }
 
 /// Generated setup class from Pigeon to handle messages through the `binaryMessenger`.
@@ -374,6 +450,27 @@ class RtcPeerConnectionHostApiSetup {
     } else {
       createAnswerChannel.setMessageHandler(nil)
     }
+    /// Generate a stats report about the PeerConnection with the given [connectionId].
+    /// If [typeFilter] is provided, the report will only contain stats matching types in that list, reducing the codec
+    /// overhead from the native platform.
+    let getStatsChannel = FlutterBasicMessageChannel(name: "dev.flutter.pigeon.live_link_vcam.RtcPeerConnectionHostApi.getStats", binaryMessenger: binaryMessenger, codec: codec)
+    if let api = api {
+      getStatsChannel.setMessageHandler { message, reply in
+        let args = message as! [Any?]
+        let connectionIdArg = args[0] is Int64 ? args[0] as! Int64 : Int64(args[0] as! Int32)
+        let typeFilterArg: [String]? = nilOrValue(args[1])
+        api.getStats(connectionId: connectionIdArg, typeFilter: typeFilterArg) { result in
+          switch result {
+            case .success(let res):
+              reply(wrapResult(res))
+            case .failure(let error):
+              reply(wrapError(error))
+          }
+        }
+      }
+    } else {
+      getStatsChannel.setMessageHandler(nil)
+    }
   }
 }
 private class RtcPeerConnectionFlutterApiCodecReader: FlutterStandardReader {
diff --git a/Engine/Extras/VirtualProduction/FlutterLiveLinkVCAM/IOS/Runner/WebRTC/Flutter/FlutterRtcPeerConnection.swift b/Engine/Extras/VirtualProduction/FlutterLiveLinkVCAM/IOS/Runner/WebRTC/Flutter/FlutterRtcPeerConnection.swift
--- a/Engine/Extras/VirtualProduction/FlutterLiveLinkVCAM/IOS/Runner/WebRTC/Flutter/FlutterRtcPeerConnection.swift
+++ b/Engine/Extras/VirtualProduction/FlutterLiveLinkVCAM/IOS/Runner/WebRTC/Flutter/FlutterRtcPeerConnection.swift
@@ -111,6 +111,34 @@ class FlutterRtcPeerConnection : IdObject {
     }
   }
   
+  /**
+   * Gather a WebRTC stats report for this connection.
+   */
+  func getStats(typeFilter: [String]?, completion: @escaping (Result<RtcStatsReport, Error>) -> Void) {
+    assert(connection != nil)
+    
+    connection!.statistics(completionHandler: { nativeReport in
+      var statsMap = Dictionary<String?, RtcStats?>()
+      
+      for (statsKey, stats) in nativeReport.statistics {
+        // Skip types that don't pass the filter
+        if (typeFilter != nil && !typeFilter!.contains(stats.type)) {
+          continue
+        }
+        
+        statsMap[statsKey] = RtcStats(
+          timestampUs: stats.timestamp_us,
+          type: stats.type,
+          id: stats.id,
+          values: stats.values
+        )
+      }
+      
+      let report = RtcStatsReport(timestampUs: nativeReport.timestamp_us, stats: statsMap)
+      completion(Result.success(report))
+    })
+  }
+  
   // MARK: IdObject
   override func onIdReady() -> Bool {
     let constraints = RTCMediaConstraints(mandatoryConstraints: nil, optionalConstraints: nil)
