[UnrealStage] Move log viewer to FlutterCommon. Update log viewer to keep limited history and display similarly to VCAM

diff --git a/Engine/Extras/VirtualProduction/EpicStageApp/lib/widgets/screens/main/toolbar/settings/pages/settings_advanced_view.dart b/Engine/Extras/VirtualProduction/EpicStageApp/lib/widgets/screens/main/toolbar/settings/pages/settings_advanced_view.dart
--- a/Engine/Extras/VirtualProduction/EpicStageApp/lib/widgets/screens/main/toolbar/settings/pages/settings_advanced_view.dart
+++ b/Engine/Extras/VirtualProduction/EpicStageApp/lib/widgets/screens/main/toolbar/settings/pages/settings_advanced_view.dart
@@ -1,12 +1,12 @@
 // Copyright Epic Games, Inc. All Rights Reserved.
 
+import 'package:epic_common/widgets.dart';
 import 'package:flutter/material.dart';
 import 'package:flutter_gen/gen_l10n/app_localizations.dart';
 import 'package:provider/provider.dart';
 import 'package:streaming_shared_preferences/streaming_shared_preferences.dart';
 
 import '../../../../../../models/settings/connection_settings.dart';
-import '../settings_generic.dart';
 import 'settings_compression_picker.dart';
 
 /// Page showing the advanced settings.
diff --git a/Engine/Extras/VirtualProduction/EpicStageApp/lib/widgets/screens/main/toolbar/settings/pages/settings_compression_picker.dart b/Engine/Extras/VirtualProduction/EpicStageApp/lib/widgets/screens/main/toolbar/settings/pages/settings_compression_picker.dart
--- a/Engine/Extras/VirtualProduction/EpicStageApp/lib/widgets/screens/main/toolbar/settings/pages/settings_compression_picker.dart
+++ b/Engine/Extras/VirtualProduction/EpicStageApp/lib/widgets/screens/main/toolbar/settings/pages/settings_compression_picker.dart
@@ -1,5 +1,6 @@
 // Copyright Epic Games, Inc. All Rights Reserved.
 
+import 'package:epic_common/widgets.dart';
 import 'package:flutter/material.dart';
 import 'package:flutter_gen/gen_l10n/app_localizations.dart';
 import 'package:provider/provider.dart';
@@ -7,7 +8,6 @@ import 'package:streaming_shared_preferences/streaming_shared_preferences.dart';
 
 import '../../../../../../models/engine_connection.dart';
 import '../../../../../../models/settings/connection_settings.dart';
-import '../settings_generic.dart';
 
 /// Page showing the possible compression modes.
 class SettingsCompressionPicker extends StatelessWidget {
diff --git a/Engine/Extras/VirtualProduction/EpicStageApp/lib/widgets/screens/main/toolbar/settings/pages/settings_dialog_main.dart b/Engine/Extras/VirtualProduction/EpicStageApp/lib/widgets/screens/main/toolbar/settings/pages/settings_dialog_main.dart
--- a/Engine/Extras/VirtualProduction/EpicStageApp/lib/widgets/screens/main/toolbar/settings/pages/settings_dialog_main.dart
+++ b/Engine/Extras/VirtualProduction/EpicStageApp/lib/widgets/screens/main/toolbar/settings/pages/settings_dialog_main.dart
@@ -12,7 +12,6 @@ import '../../../../../../models/navigator_keys.dart';
 import '../../../../../../models/settings/selected_actor_settings.dart';
 import '../../../../../../models/unreal_actor_manager.dart';
 import '../../../../eula/eula_screen.dart';
-import '../settings_generic.dart';
 import 'settings_advanced_view.dart';
 import 'settings_log_list.dart';
 import 'settings_root_actor_picker.dart';
diff --git a/Engine/Extras/VirtualProduction/EpicStageApp/lib/widgets/screens/main/toolbar/settings/pages/settings_log_list.dart b/Engine/Extras/VirtualProduction/EpicStageApp/lib/widgets/screens/main/toolbar/settings/pages/settings_log_list.dart
--- a/Engine/Extras/VirtualProduction/EpicStageApp/lib/widgets/screens/main/toolbar/settings/pages/settings_log_list.dart
+++ b/Engine/Extras/VirtualProduction/EpicStageApp/lib/widgets/screens/main/toolbar/settings/pages/settings_log_list.dart
@@ -8,7 +8,6 @@ import 'package:flutter/material.dart';
 import 'package:flutter_gen/gen_l10n/app_localizations.dart';
 
 import '../../../../../../utilities/guarded_refresh_state.dart';
-import '../settings_generic.dart';
 import 'settings_log_view.dart';
 
 /// Page showing the list of log files.
@@ -43,17 +42,8 @@ class _SettingsLogListState extends State<SettingsLogList> with GuardedRefreshSt
 
   /// Update the list of log files.
   void _updateLogFileList() async {
-    final Directory logDirectory = await Logging.instance.getLogDirectory();
-    final List<File> files =
-        await logDirectory.list(recursive: false).map((FileSystemEntity entity) => File(entity.path)).toList();
-
-    // Sort so the most recent files appear first. We assume that files are named identically aside from their dates,
-    // which are assumed to be in lexicographical order (YYYY_MM_DD)
-    files.sort((File fileA, File fileB) => -fileA.path.compareTo(fileB.path));
-
-    setState(() {
-      _logFiles = files;
-    });
+    _logFiles = await Logging.getAppLogFiles();
+    setState(() {});
   }
 
   @override
@@ -70,14 +60,20 @@ class _SettingsLogListState extends State<SettingsLogList> with GuardedRefreshSt
             itemCount: _logFiles.length,
             itemBuilder: (context, index) {
               final File file = _logFiles[index];
+              final bool bIsCurrent = index == 0;
 
               return SettingsMenuItem(
-                title: Logging.getNameForLog(file),
+                title: Logging.makeNameForLog(
+                  context: context,
+                  logFile: file,
+                  bIsCurrent: bIsCurrent,
+                ),
                 iconPath: 'packages/epic_common/assets/icons/log.svg',
                 onTap: () => Navigator.of(context).pushNamed(
                   SettingsLogView.route,
                   arguments: SettingsLogViewArguments(
                     file: file,
+                    bIsCurrent: bIsCurrent,
                   ),
                 ),
               );
diff --git a/Engine/Extras/VirtualProduction/EpicStageApp/lib/widgets/screens/main/toolbar/settings/pages/settings_log_view.dart b/Engine/Extras/VirtualProduction/EpicStageApp/lib/widgets/screens/main/toolbar/settings/pages/settings_log_view.dart
--- a/Engine/Extras/VirtualProduction/EpicStageApp/lib/widgets/screens/main/toolbar/settings/pages/settings_log_view.dart
+++ b/Engine/Extras/VirtualProduction/EpicStageApp/lib/widgets/screens/main/toolbar/settings/pages/settings_log_view.dart
@@ -7,16 +7,17 @@ import 'package:epic_common/widgets.dart';
 import 'package:flutter/material.dart';
 import 'package:share_plus/share_plus.dart';
 
-import '../../../../../../utilities/external_notifier.dart';
 import '../../../../../../utilities/guarded_refresh_state.dart';
-import '../settings_generic.dart';
 
 /// Arguments passed to the route when a [SettingsLogView] is pushed to the navigation stack.
 class SettingsLogViewArguments {
-  const SettingsLogViewArguments({required this.file});
+  const SettingsLogViewArguments({required this.file, required this.bIsCurrent});
 
   /// The file to open.
   final File file;
+
+  /// Whether this is the current log.
+  final bool bIsCurrent;
 }
 
 /// Page showing the full view of a log.
@@ -37,7 +38,11 @@ class _SettingsLogViewState extends State<SettingsLogView> with GuardedRefreshSt
     final arguments = ModalRoute.of(context)?.settings.arguments as SettingsLogViewArguments;
 
     return SettingsPageScaffold(
-      title: Logging.getNameForLog(arguments.file),
+      title: Logging.makeNameForLog(
+        context: context,
+        logFile: arguments.file,
+        bIsCurrent: arguments.bIsCurrent,
+      ),
       titleBarTrailing: EpicIconButton(
         key: shareButtonKey,
         iconPath: 'packages/epic_common/assets/icons/share.svg',
@@ -46,12 +51,7 @@ class _SettingsLogViewState extends State<SettingsLogView> with GuardedRefreshSt
       ),
       body: SizedBox(
         height: MediaQuery.of(context).size.height - 120,
-        child: FutureBuilder(
-          future: arguments.file.readAsString(),
-          builder: (context, AsyncSnapshot<String> snapshot) {
-            return _LogText(text: snapshot.data);
-          },
-        ),
+        child: LogFileViewer(arguments.file),
       ),
     );
   }
@@ -73,319 +73,3 @@ class _SettingsLogViewState extends State<SettingsLogView> with GuardedRefreshSt
     );
   }
 }
-
-/// Widget that displays text from a log.
-class _LogText extends StatefulWidget {
-  const _LogText({Key? key, required this.text}) : super(key: key);
-
-  /// Text to display.
-  final String? text;
-
-  @override
-  State<_LogText> createState() => _LogTextState();
-}
-
-class _LogTextState extends State<_LogText> {
-  /// Vertical scroll controller
-  final _verticalScroll = ScrollController();
-
-  /// Horizontal scroll controller
-  final _horizontalScroll = ScrollController();
-
-  /// Notifier for when the text is changed
-  final _textChangeNotifier = ExternalNotifier();
-
-  /// Text style to use for log text
-  late TextStyle _textStyle;
-
-  /// List of lines to display in the text
-  List<String> _lines = [];
-
-  /// Width of the text scrolling area
-  double? _textWidth;
-
-  @override
-  void didChangeDependencies() {
-    super.didChangeDependencies();
-
-    _textStyle = Theme.of(context).textTheme.bodyMedium!.copyWith(
-          fontFamily: 'Droid Sans Mono',
-          fontSize: 12,
-        );
-
-    _updateText();
-  }
-
-  @override
-  void didUpdateWidget(covariant _LogText oldWidget) {
-    super.didUpdateWidget(oldWidget);
-
-    if (oldWidget.text != widget.text) {
-      _updateText();
-
-      WidgetsBinding.instance.addPostFrameCallback((_) {
-        _verticalScroll.jumpTo(_verticalScroll.initialScrollOffset);
-        _horizontalScroll.jumpTo(_horizontalScroll.initialScrollOffset);
-
-        _textChangeNotifier.notifyListeners();
-      });
-    }
-  }
-
-  @override
-  Widget build(BuildContext context) {
-    if (_textWidth == null) {
-      return const _SettingsLogViewLoadingSpinner();
-    }
-
-    final scrollArea = MediaQuery.removePadding(
-      context: context,
-      removeTop: true,
-      removeBottom: true,
-      child: DefaultTextStyle(
-        style: _textStyle,
-        child: EpicScrollBar(
-          padding: EdgeInsets.only(left: 16, right: 8, top: 8, bottom: 8),
-          controller: _verticalScroll,
-          // Listen to the horizontal scroll view, nested 1 level past the nearest view
-          notification: (notification) => notification.depth == 1,
-          child: Padding(
-            padding: const EdgeInsets.only(right: 8),
-            child: EpicScrollBar(
-              padding: EdgeInsets.only(right: 8),
-              controller: _horizontalScroll,
-              child: SingleChildScrollView(
-                padding: EdgeInsets.only(bottom: 16),
-                scrollDirection: Axis.horizontal,
-                controller: _horizontalScroll,
-                child: SizedBox(
-                  width: _textWidth,
-                  child: ScrollConfiguration(
-                    behavior: ScrollConfiguration.of(context).copyWith(scrollbars: false),
-
-                    // Display text as a ListView of lines for better performance. Rather than laying out all the text,
-                    // this only lays out and renders the visible lines.
-                    child: ListView.builder(
-                      shrinkWrap: true,
-                      scrollDirection: Axis.vertical,
-                      controller: _verticalScroll,
-                      itemCount: _lines.length,
-                      itemBuilder: (_, final int lineIndex) => Text(
-                        _lines[lineIndex],
-                        softWrap: false,
-                      ),
-                      prototypeItem: const Text(''),
-                    ),
-                  ),
-                ),
-              ),
-            ),
-          ),
-        ),
-      ),
-    );
-
-    // Add buttons to quickly scroll to top/bottom
-    return Stack(
-      clipBehavior: Clip.none,
-      children: [
-        scrollArea,
-        Padding(
-          padding: EdgeInsets.symmetric(vertical: 10),
-          child: Stack(
-            clipBehavior: Clip.none,
-            children: [
-              Align(
-                alignment: Alignment.topCenter,
-                child: _ScrollButton(
-                  bScrollToStart: true,
-                  scrollController: _verticalScroll,
-                  resetNotifier: _textChangeNotifier,
-                ),
-              ),
-              Align(
-                alignment: Alignment.bottomCenter,
-                child: _ScrollButton(
-                  bScrollToStart: false,
-                  scrollController: _verticalScroll,
-                  resetNotifier: _textChangeNotifier,
-                ),
-              ),
-            ],
-          ),
-        ),
-      ],
-    );
-  }
-
-  /// Update information based on the text to display.
-  void _updateText() async {
-    if (widget.text == null) {
-      _lines = const [];
-      _textWidth = null;
-      return;
-    }
-
-    _lines = widget.text!.split('\n').where((line) => line.length > 0).toList(growable: false);
-
-    // Find the width of the longest line to determine our horizontal scroll area.
-    // This only works if we use a monospace font, but is much cheaper than measuring the entire text for long logs.
-    final String longestLine = _lines.fold(
-      '',
-      (String longest, String line) => line.length > longest.length ? line : longest,
-    );
-
-    final textPainter = TextPainter(
-      textDirection: Directionality.of(context),
-      text: TextSpan(
-        text: longestLine,
-        style: _textStyle,
-      ),
-    );
-
-    textPainter.layout();
-
-    _textWidth = textPainter.width;
-  }
-}
-
-/// A button that lets the user instantly scroll to the start or end and hides itself when not relevant.
-class _ScrollButton extends StatefulWidget {
-  const _ScrollButton({
-    Key? key,
-    required this.scrollController,
-    required this.bScrollToStart,
-    this.resetNotifier,
-  }) : super(key: key);
-
-  /// Controller used to scroll the controlled view.
-  final ScrollController scrollController;
-
-  /// Whether this should scroll to the start (true) or end (false) of the scroll view.
-  final bool bScrollToStart;
-
-  /// If provided, update the button's visibility when this notifier fires (e.g. scroll bounds have changed).
-  final ChangeNotifier? resetNotifier;
-
-  @override
-  State<_ScrollButton> createState() => _ScrollButtonState();
-}
-
-class _ScrollButtonState extends State<_ScrollButton> {
-  bool _bIsVisible = false;
-  bool _bIsAnimating = false;
-
-  @override
-  void initState() {
-    super.initState();
-    widget.scrollController.addListener(_updateVisibility);
-    widget.resetNotifier?.addListener(_updateVisibility);
-
-    WidgetsBinding.instance.addPostFrameCallback((_) {
-      _updateVisibility();
-    });
-  }
-
-  @override
-  void dispose() {
-    super.dispose();
-
-    widget.scrollController.removeListener(_updateVisibility);
-    widget.resetNotifier?.removeListener(_updateVisibility);
-  }
-
-  @override
-  void didUpdateWidget(covariant _ScrollButton oldWidget) {
-    super.didUpdateWidget(oldWidget);
-
-    oldWidget.resetNotifier?.removeListener(_updateVisibility);
-    oldWidget.scrollController.removeListener(_updateVisibility);
-
-    widget.resetNotifier?.addListener(_updateVisibility);
-    widget.scrollController.addListener(_updateVisibility);
-  }
-
-  @override
-  Widget build(BuildContext context) {
-    return AnimatedOpacity(
-      opacity: (_bIsVisible && !_bIsAnimating) ? 1 : 0,
-      duration: Duration(milliseconds: 150),
-      child: FloatingActionButton(
-        child: AssetIcon(
-          path: widget.bScrollToStart
-              ? 'packages/epic_common/assets/icons/chevron_up.svg'
-              : 'packages/epic_common/assets/icons/chevron_down.svg',
-          size: 32,
-        ),
-        onPressed: _onPressed,
-        heroTag: null,
-      ),
-    );
-  }
-
-  /// Update whether to show the button based on scroll position.
-  void _updateVisibility() {
-    bool bNewIsVisible = false;
-    final double? scrollTarget = _getScrollTarget();
-    if (scrollTarget == null) {
-      bNewIsVisible = false;
-    } else {
-      bNewIsVisible = widget.scrollController.offset != scrollTarget;
-    }
-
-    if (bNewIsVisible != _bIsVisible) {
-      setState(() {
-        _bIsVisible = bNewIsVisible;
-      });
-    }
-  }
-
-  /// Called when the button is pressed.
-  void _onPressed() {
-    final double? scrollTarget = _getScrollTarget();
-    if (scrollTarget != null) {
-      setState(() {
-        _bIsAnimating = true;
-      });
-
-      widget.scrollController
-          .animateTo(
-            scrollTarget,
-            duration: Duration(milliseconds: 300),
-            curve: Curves.easeInOutCubic,
-          )
-          .then((_) => setState(() {
-                _bIsAnimating = false;
-              }));
-    }
-  }
-
-  /// Get the position to scroll to when the button is pressed.
-  double? _getScrollTarget() {
-    final scrollPosition = widget.scrollController.position;
-    if (!scrollPosition.hasContentDimensions) {
-      return null;
-    }
-
-    return widget.bScrollToStart ? scrollPosition.minScrollExtent : scrollPosition.maxScrollExtent;
-  }
-}
-
-/// Contents to display while the log view is loading.
-class _SettingsLogViewLoadingSpinner extends StatelessWidget {
-  const _SettingsLogViewLoadingSpinner({Key? key}) : super(key: key);
-
-  @override
-  Widget build(BuildContext context) {
-    return Center(
-      child: const Center(
-        child: SizedBox.square(
-          dimension: 120,
-          child: CircularProgressIndicator(
-            strokeWidth: 8,
-          ),
-        ),
-      ),
-    );
-  }
-}
diff --git a/Engine/Extras/VirtualProduction/EpicStageApp/lib/widgets/screens/main/toolbar/settings/pages/settings_root_actor_picker.dart b/Engine/Extras/VirtualProduction/EpicStageApp/lib/widgets/screens/main/toolbar/settings/pages/settings_root_actor_picker.dart
--- a/Engine/Extras/VirtualProduction/EpicStageApp/lib/widgets/screens/main/toolbar/settings/pages/settings_root_actor_picker.dart
+++ b/Engine/Extras/VirtualProduction/EpicStageApp/lib/widgets/screens/main/toolbar/settings/pages/settings_root_actor_picker.dart
@@ -1,5 +1,6 @@
 // Copyright Epic Games, Inc. All Rights Reserved.
 
+import 'package:epic_common/widgets.dart';
 import 'package:flutter/material.dart';
 import 'package:flutter_gen/gen_l10n/app_localizations.dart';
 import 'package:provider/provider.dart';
@@ -9,7 +10,6 @@ import '../../../../../../models/unreal_actor_manager.dart';
 import '../../../../../../models/unreal_types.dart';
 import '../../../../../../utilities/constants.dart';
 import '../../../../../../utilities/guarded_refresh_state.dart';
-import '../settings_generic.dart';
 
 /// Page showing the list of nDisplay root actors to pick from.
 class SettingsDialogRootActorPicker extends StatefulWidget {
diff --git a/Engine/Extras/VirtualProduction/EpicStageApp/lib/widgets/screens/main/toolbar/settings/settings_generic.dart b/Engine/Extras/VirtualProduction/EpicStageApp/lib/widgets/screens/main/toolbar/settings/settings_generic.dart
deleted file mode 100644
