[LiveLinkVCAM] Localization support

diff --git a/Engine/Extras/VirtualProduction/FlutterLiveLinkVCAM/pubspec.yaml b/Engine/Extras/VirtualProduction/FlutterLiveLinkVCAM/pubspec.yaml
--- a/Engine/Extras/VirtualProduction/FlutterLiveLinkVCAM/pubspec.yaml
+++ b/Engine/Extras/VirtualProduction/FlutterLiveLinkVCAM/pubspec.yaml
@@ -40,6 +40,9 @@ dev_dependencies:
 flutter:
   uses-material-design: true
 
+  # Generate localization files
+  generate: true
+
   fonts:
     - family: Inter
       fonts:
@@ -48,3 +51,12 @@ flutter:
     - family: Droid Sans Mono
       fonts:
         - asset: packages/epic_common/assets/fonts/DroidSansMono.ttf
+
+# Automatically generate arb localization files from a CSV
+arb_generator:
+  input_filepath: "lib/l10n/l10n.csv"
+  output_directory: "lib/l10n"
+  filename_prepend: "app_"
+  csv_settings:
+    description_index: 1
+    base_index: 2
