[LiveLinkVCAM] Add missing file for CL 29817961

diff --git a/Engine/Extras/VirtualProduction/FlutterLiveLinkVCAM/l10n.yaml b/Engine/Extras/VirtualProduction/FlutterLiveLinkVCAM/l10n.yaml
new file mode 100644
--- /dev/null
+++ b/Engine/Extras/VirtualProduction/FlutterLiveLinkVCAM/l10n.yaml
@@ -0,0 +1,3 @@
+arb-dir: lib/l10n
+template-arb-file: app_en.arb
+output-localization-file: app_localizations.dart
\ No newline at end of file
