Live Link VCAM iOS v1.3.2 build 9 : updated mic permission string (en only)

diff --git a/Engine/Extras/VirtualProduction/LiveLinkVCAM/vcam/Info.plist b/Engine/Extras/VirtualProduction/LiveLinkVCAM/vcam/Info.plist
--- a/Engine/Extras/VirtualProduction/LiveLinkVCAM/vcam/Info.plist
+++ b/Engine/Extras/VirtualProduction/LiveLinkVCAM/vcam/Info.plist
@@ -42,7 +42,7 @@
 	<key>NSLocalNetworkUsageDescription</key>
 	<string>VCAM transmits &amp; receives data from Unreal Engine.</string>
 	<key>NSMicrophoneUsageDescription</key>
-	<string>We don't use the microphone.</string>
+	<string>Your audio will be transmitted to your Unreal Engine workstation for processing.</string>
 	<key>UIApplicationSceneManifest</key>
 	<dict>
 		<key>UIApplicationSupportsMultipleScenes</key>
