[LiveLinkVCAM] Add WebRTC streaming audio support

diff --git a/Engine/Extras/VirtualProduction/FlutterLiveLinkVCAM/IOS/Runner/WebRTC/API/FlutterRtcPeerConnectionApi.swift b/Engine/Extras/VirtualProduction/FlutterLiveLinkVCAM/IOS/Runner/WebRTC/API/FlutterRtcPeerConnectionApi.swift
--- a/Engine/Extras/VirtualProduction/FlutterLiveLinkVCAM/IOS/Runner/WebRTC/API/FlutterRtcPeerConnectionApi.swift
+++ b/Engine/Extras/VirtualProduction/FlutterLiveLinkVCAM/IOS/Runner/WebRTC/API/FlutterRtcPeerConnectionApi.swift
@@ -10,6 +10,21 @@ class FlutterRtcPeerConnectionApi: FlutterPluginApi<RtcPeerConnectionFlutterApi>
     
     RtcPeerConnectionHostApiSetup.setUp(binaryMessenger: binaryMessenger, api: self)
     FlutterRtcPeerConnectionApi.instance = self
+    
+    self.audioQueue.async { [weak self] in
+        guard let self = self else {
+            return
+        }
+        
+        self.rtcAudioSession.lockForConfiguration()
+        do {
+            try self.rtcAudioSession.setCategory(AVAudioSession.Category.ambient)
+            try self.rtcAudioSession.setMode(AVAudioSession.Mode.default)
+        } catch let error {
+            debugPrint("Error setting AVAudioSession category: \(error)")
+        }
+        self.rtcAudioSession.unlockForConfiguration()
+    }
   }
   
   deinit {
@@ -33,6 +48,12 @@ class FlutterRtcPeerConnectionApi: FlutterPluginApi<RtcPeerConnectionFlutterApi>
   /// Video decoder used for incoming video streams.
   private let videoDecoderFactory = RTCDefaultVideoDecoderFactory()
   
+  /// Dispatch queue used for async audio system calls.
+  private let audioQueue = DispatchQueue(label: "audio")
+  
+  /// The audio session used to handle WebRTC audio data.
+  private let rtcAudioSession =  RTCAudioSession.sharedInstance()
+  
   /// Factory used to create PeerConnections as requested by the client.
   private var _peerConnectionFactory: RTCPeerConnectionFactory?
   private var peerConnectionFactory: RTCPeerConnectionFactory {
