[UnrealStage] [FlutterCommon] Refactor beacon sender to a generic base class

diff --git a/Engine/Extras/FlutterCommon/pubspec.yaml b/Engine/Extras/FlutterCommon/pubspec.yaml
--- a/Engine/Extras/FlutterCommon/pubspec.yaml
+++ b/Engine/Extras/FlutterCommon/pubspec.yaml
@@ -20,6 +20,8 @@ dependencies:
   path_provider: ^2.0.11
   provider: ^6.0.2
   streaming_shared_preferences: ^2.0.0
+  uuid: ^3.0.0
+  archive: ^3.2.0
 
 dev_dependencies:
   flutter_test:
diff --git a/Engine/Extras/VirtualProduction/EpicStageApp/lib/models/engine_connection.dart b/Engine/Extras/VirtualProduction/EpicStageApp/lib/models/engine_connection.dart
--- a/Engine/Extras/VirtualProduction/EpicStageApp/lib/models/engine_connection.dart
+++ b/Engine/Extras/VirtualProduction/EpicStageApp/lib/models/engine_connection.dart
@@ -12,7 +12,6 @@ import 'package:flutter_gen/gen_l10n/app_localizations.dart';
 import 'package:logging/logging.dart';
 import 'package:provider/provider.dart';
 import 'package:shared_preferences/shared_preferences.dart';
-import 'package:uuid/uuid.dart';
 import 'package:web_socket_channel/io.dart';
 
 import './navigator_keys.dart';
@@ -877,7 +876,6 @@ class EngineConnectionManager with WidgetsBindingObserver {
       ConnectionData? loadedConnectionData;
       try {
         loadedConnectionData = ConnectionData(
-          uuid: UuidValue(const Uuid().v1()),
           name: name,
           websocketAddress: InternetAddress(address),
           websocketPort: port,
diff --git a/Engine/Extras/VirtualProduction/EpicStageApp/lib/models/unreal_engine_beacon.dart b/Engine/Extras/VirtualProduction/EpicStageApp/lib/models/unreal_engine_beacon.dart
deleted file mode 100644
diff --git a/Engine/Extras/VirtualProduction/EpicStageApp/lib/utilities/net_utilities.dart b/Engine/Extras/VirtualProduction/EpicStageApp/lib/utilities/net_utilities.dart
--- a/Engine/Extras/VirtualProduction/EpicStageApp/lib/utilities/net_utilities.dart
+++ b/Engine/Extras/VirtualProduction/EpicStageApp/lib/utilities/net_utilities.dart
@@ -2,73 +2,66 @@
 
 import 'dart:async';
 import 'dart:io';
-import 'dart:typed_data';
 
-import 'package:archive/archive_io.dart';
+import 'package:archive/archive.dart';
+import 'package:epic_common/unreal_beacon.dart';
 import 'package:logging/logging.dart';
-import 'package:uuid/uuid.dart';
 
 import '../models/engine_connection.dart';
 
-final _log = Logger('NetUtilities');
+/// The configuration to use for the beacon that detects compatible Unreal Engine instances.
+final unrealEngineBeaconConfig = UnrealEngineBeaconConfig(
+  beaconAddress: InternetAddress('230.0.0.2'),
+  beaconPort: 6667,
+  protocolIdentifier: 'ES@p',
+  userDataReader: (InputStream inputStream, InternetAddress address, _, __) {
+    final int port = inputStream.readUint32();
+    String? name = readStringFromStream(inputStream);
+
+    // If the engine didn't provide a friendly name, just use the address and WebSocket port.
+    if (name == null || name.isEmpty || (name.length == 1 && name[0] == '\x00')) {
+      name = '${address.address}:${port.toString()}';
+    }
 
-const int _protocolVersion = 0; // The current protocol version for beacon messages
-const String _protocolIdentifier = 'ES@p';
+    return UnrealBeaconResponseUserData(
+      additionalData: UnrealStageBeaconData(websocketPort: port),
+      name: name,
+    );
+  },
+);
 
-/// Data about a detected engine connection
-class ConnectionData {
-  ConnectionData({required this.uuid, required this.name, required this.websocketAddress, required this.websocketPort})
-      : lastSeenTime = DateTime.now();
-
-  final UuidValue uuid;
-  String name;
-  InternetAddress websocketAddress;
-  int websocketPort;
-  DateTime lastSeenTime;
-}
+/// Type of response received from beacon messages for this app.
+typedef UnrealStageBeaconResponse = UnrealBeaconResponse<UnrealStageBeaconData>;
 
-/// Convert a Uint8List to a UUID.
-UuidValue uuidFromUint8List(Uint8List data) {
-  String guidString = Uuid.unparse(data);
-  return UuidValue(guidString, false);
-}
+/// Additional Unreal Stage-specific data received from beacon messages.
+class UnrealStageBeaconData {
+  const UnrealStageBeaconData({required this.websocketPort});
 
-/// Make the beacon message that will be multicast to find compatible engine instances.
-Uint8List makeBeaconMessage() {
-  return Uint8List.fromList(_protocolIdentifier.codeUnits + [_protocolVersion]);
+  /// Port to connect to for WebSocket communication.
+  final int websocketPort;
 }
 
-/// Retrieve the connection data from the datagram received from the response to a beacon message.
-/// If the connection data couldn't be retrieved, this will return null.
-ConnectionData? getConnectionFromBeaconResponse(Datagram datagram) {
-  final InputStream stream = InputStream(datagram.data);
-  stream.byteOrder = LITTLE_ENDIAN;
-  ConnectionData? connection;
-
-  try {
-    stream.readByte(); // Protocol version; currently ignored
-    final UuidValue engineUuid = uuidFromUint8List(stream.readBytes(16).toUint8List());
-    final int port = stream.readUint32();
-
-    final int nameLength = stream.readUint32();
-    String name = stream.readString(size: nameLength, utf8: true);
-
-    if (name[name.length - 1] == '\x00') {
-      // Remove the null terminator
-      name = name.substring(0, name.length - 1);
-    }
+/// Data about a connection to Unreal Engine.
+class ConnectionData {
+  const ConnectionData({required this.name, required this.websocketAddress, required this.websocketPort});
+
+  /// Create connection data from a [response] to an UnrealEngineBeacon message.
+  static ConnectionData fromBeaconResponse(UnrealStageBeaconResponse response) {
+    return ConnectionData(
+      name: response.name,
+      websocketAddress: response.address,
+      websocketPort: response.additionalData.websocketPort,
+    );
+  }
 
-    // If the engine didn't provide a friendly name, just use the address and port.
-    if (name.isEmpty || (name.length == 1 && name[0] == '\x00')) {
-      name = '${datagram.address.address}:${port.toString()}';
-    }
+  /// User-friendly name of the connection.
+  final String name;
 
-    connection = ConnectionData(uuid: engineUuid, name: name, websocketAddress: datagram.address, websocketPort: port);
-  } catch (error) {
-    _log.warning('Failed to interpret beacon response from ${datagram.address.address}:/n$error');
-  }
+  /// Address to connect to for WebSocket communication.
+  final InternetAddress websocketAddress;
 
-  return connection;
+  /// Port to connect to for WebSocket communication.
+  final int websocketPort;
 }
 
 /// A helper class that automatically retries sending a message to the engine if it times out.
diff --git a/Engine/Extras/VirtualProduction/EpicStageApp/lib/widgets/screens/connect/connect.dart b/Engine/Extras/VirtualProduction/EpicStageApp/lib/widgets/screens/connect/connect.dart
--- a/Engine/Extras/VirtualProduction/EpicStageApp/lib/widgets/screens/connect/connect.dart
+++ b/Engine/Extras/VirtualProduction/EpicStageApp/lib/widgets/screens/connect/connect.dart
@@ -1,6 +1,9 @@
 // Copyright Epic Games, Inc. All Rights Reserved.
 
+import 'dart:collection';
+
 import 'package:epic_common/theme.dart';
+import 'package:epic_common/unreal_beacon.dart';
 import 'package:epic_common/widgets.dart';
 import 'package:flutter/material.dart';
 import 'package:flutter_gen/gen_l10n/app_localizations.dart';
@@ -10,7 +13,6 @@ import '../../../../widgets/screens/connect/views/connection_list.dart';
 import '../../../../widgets/screens/main/toolbar/main_screen_toolbar.dart';
 import '../../../models/engine_connection.dart';
 import '../../../models/navigator_keys.dart';
-import '../../../models/unreal_engine_beacon.dart';
 import '../../../utilities/debug_utilities.dart';
 import '../../../utilities/net_utilities.dart';
 import 'mixins/connect_mixin.dart';
@@ -30,7 +32,7 @@ class ConnectScreen extends StatefulWidget {
 
 class _ConnectScreenState extends State<ConnectScreen> with WidgetsBindingObserver, ConnectMixin {
   /// Beacon which will be used to detect instances of Unreal Engine.
-  late final UnrealEngineBeacon _engineBeacon;
+  late final UnrealEngineBeacon<UnrealStageBeaconData> _engineBeacon;
 
   /// Data for the last engine we were connected to, if any.
   ConnectionData? _lastConnection;
@@ -55,8 +57,8 @@ class _ConnectScreenState extends State<ConnectScreen> with WidgetsBindingObserv
 
     _engineBeacon = UnrealEngineBeacon(
       context: context,
+      config: unrealEngineBeaconConfig,
       onBeaconFailure: _onBeaconFailure,
-      onEngineInstancesChanged: () => setState(() {}),
     );
 
     Provider.of<EngineConnectionManager>(context, listen: false)
@@ -137,19 +139,23 @@ class _ConnectScreenState extends State<ConnectScreen> with WidgetsBindingObserv
                         iconPath: 'packages/epic_common/assets/icons/unreal_u_logo.svg',
                       ),
                       Expanded(
-                        child: _engineBeacon.connections.length > 0
-                            ? ConnectionList(
-                                connections: _engineBeacon.connections,
-                              )
-                            : EmptyPlaceholder(
-                                message: AppLocalizations.of(context)!.connectScreenAllConnectionsPanelEmptyMessage,
-                                button: EpicWideButton(
-                                  text: AppLocalizations.of(context)!.connectScreenAllConnectionsPanelEmptyButtonLabel,
-                                  iconPath: 'packages/epic_common/assets/icons/plus.svg',
-                                  color: UnrealColors.highlightGreen,
-                                  onPressed: _showManualConnectDialog,
+                        child: ValueListenableBuilder<UnmodifiableListView<UnrealStageBeaconResponse>>(
+                          valueListenable: _engineBeacon.responses,
+                          builder: (_, responses, __) => responses.isNotEmpty
+                              ? ConnectionList(
+                                  connections: responses.map(ConnectionData.fromBeaconResponse).toList(growable: false),
+                                )
+                              : EmptyPlaceholder(
+                                  message: AppLocalizations.of(context)!.connectScreenAllConnectionsPanelEmptyMessage,
+                                  button: EpicWideButton(
+                                    text:
+                                        AppLocalizations.of(context)!.connectScreenAllConnectionsPanelEmptyButtonLabel,
+                                    iconPath: 'packages/epic_common/assets/icons/plus.svg',
+                                    color: UnrealColors.highlightGreen,
+                                    onPressed: _showManualConnectDialog,
+                                  ),
                                 ),
-                              ),
+                        ),
                       ),
                     ],
                   ),
diff --git a/Engine/Extras/VirtualProduction/EpicStageApp/lib/widgets/screens/connect/views/manual_connect_form.dart b/Engine/Extras/VirtualProduction/EpicStageApp/lib/widgets/screens/connect/views/manual_connect_form.dart
--- a/Engine/Extras/VirtualProduction/EpicStageApp/lib/widgets/screens/connect/views/manual_connect_form.dart
+++ b/Engine/Extras/VirtualProduction/EpicStageApp/lib/widgets/screens/connect/views/manual_connect_form.dart
@@ -6,7 +6,6 @@ import 'package:epic_common/localizations.dart';
 import 'package:epic_common/widgets.dart';
 import 'package:flutter/material.dart';
 import 'package:flutter_gen/gen_l10n/app_localizations.dart';
-import 'package:uuid/uuid.dart';
 
 import '../../../../utilities/constants.dart';
 import '../../../../utilities/net_utilities.dart';
@@ -103,7 +102,6 @@ class _ManualConnectFormState extends State<ManualConnectForm> {
     Navigator.of(context).pop(false);
 
     ConnectionData connection = ConnectionData(
-      uuid: UuidValue(const Uuid().v1()),
       name: AppLocalizations.of(context)!.connectScreenManualConnectionLabel,
       websocketAddress: InternetAddress(_ipTextController.text),
       websocketPort: int.parse(_portTextController.text),
