[UnrealStage] Fix handling of prematurely ended transactions

diff --git a/Engine/Extras/VirtualProduction/EpicStageApp/lib/models/unreal_property_manager.dart b/Engine/Extras/VirtualProduction/EpicStageApp/lib/models/unreal_property_manager.dart
--- a/Engine/Extras/VirtualProduction/EpicStageApp/lib/models/unreal_property_manager.dart
+++ b/Engine/Extras/VirtualProduction/EpicStageApp/lib/models/unreal_property_manager.dart
@@ -650,7 +650,10 @@ class UnrealPropertyManager {
           final beginMessage = _transactionManager.createBeginTransactionMessage(
             transaction.description,
             // If the transaction ends prematurely, flag it to be re-created
-            prematureEndCallback: () => transaction.bWasCreatedInEngine = false,
+            prematureEndCallback: () {
+              _activeTransactionId = UnrealTransactionManager.invalidId;
+              transaction.bWasCreatedInEngine = false;
+            },
           );
           assert(beginMessage != null);
           messages.add(beginMessage);
diff --git a/Engine/Extras/VirtualProduction/EpicStageApp/lib/models/unreal_transaction_manager.dart b/Engine/Extras/VirtualProduction/EpicStageApp/lib/models/unreal_transaction_manager.dart
--- a/Engine/Extras/VirtualProduction/EpicStageApp/lib/models/unreal_transaction_manager.dart
+++ b/Engine/Extras/VirtualProduction/EpicStageApp/lib/models/unreal_transaction_manager.dart
@@ -89,7 +89,7 @@ class UnrealTransactionManager {
   dynamic createEndTransactionMessage() {
     if (_activeTransaction == null) {
       _log.warning('Tried to end transaction, but no transaction was in progress');
-      return false;
+      return null;
     }
 
     final message = createUnrealWebSocketMessage('transaction.end', {
@@ -134,9 +134,12 @@ class UnrealTransactionManager {
       return;
     }
 
-    if (_activeTransaction?.id == id) {
-      _activeTransaction?.prematureEndCallback?.call();
+    if (_activeTransaction?.id != id) {
+      return;
     }
+
+    _activeTransaction?.prematureEndCallback?.call();
+    _activeTransaction = null;
   }
 }
 
