[UnrealStage] Fix thin line appearing under sliders in some cases

diff --git a/Engine/Extras/VirtualProduction/EpicStageApp/lib/widgets/elements/delta_slider.dart b/Engine/Extras/VirtualProduction/EpicStageApp/lib/widgets/elements/delta_slider.dart
--- a/Engine/Extras/VirtualProduction/EpicStageApp/lib/widgets/elements/delta_slider.dart
+++ b/Engine/Extras/VirtualProduction/EpicStageApp/lib/widgets/elements/delta_slider.dart
@@ -4,6 +4,7 @@ import 'dart:math' as math;
 
 import 'package:epic_common/preferences.dart';
 import 'package:epic_common/theme.dart';
+import 'package:epic_common/utilities/drawing_utils.dart';
 import 'package:epic_common/widgets.dart';
 import 'package:flutter/material.dart';
 import 'package:flutter_gen/gen_l10n/app_localizations.dart';
@@ -549,6 +550,8 @@ class _DeltaSliderSliderPainter extends CustomPainter {
     const double ringThickness = 2.8;
     const double disabledOpacity = 0.4;
 
+    pixelAlignCanvas(canvas);
+
     final Rect canvasRect = Offset.zero & size;
     canvas.saveLayer(canvasRect, Paint());
 
@@ -602,7 +605,6 @@ class _DeltaSliderSliderPainter extends CustomPainter {
     // Erase the mask around the rings
     final Paint ringMaskPaint = Paint()
       ..style = PaintingStyle.fill
-      ..color = Colors.transparent
       ..blendMode = BlendMode.clear;
 
     const double trackRingPadding = 5.5;
@@ -622,8 +624,7 @@ class _DeltaSliderSliderPainter extends CustomPainter {
 
     final Paint ringFillPaint = Paint()
       ..style = PaintingStyle.fill
-      ..color = Colors.white
-      ..blendMode = BlendMode.dstOut;
+      ..blendMode = BlendMode.clear;
 
     for (final Rect ringRect in ringRects) {
       canvas.drawArc(ringRect, 0, math.pi * 2, false, ringFillPaint);
