[FlutterCommon] [UnrealStage] [LiveLinkVCAM] Update dependencies and fix widget rendering on newer Flutter versions

diff --git a/Engine/Extras/FlutterCommon/pubspec.yaml b/Engine/Extras/FlutterCommon/pubspec.yaml
--- a/Engine/Extras/FlutterCommon/pubspec.yaml
+++ b/Engine/Extras/FlutterCommon/pubspec.yaml
@@ -4,7 +4,7 @@ version: 0.0.1
 
 environment:
   sdk: ">=3.0.6 <4.0.0"
-  flutter: "^3.7.0"
+  flutter: "^3.16.3"
 
 dependencies:
   flutter:
@@ -14,19 +14,19 @@ dependencies:
     sdk: flutter
 
   arb_generator: ^0.0.3
-  flutter_svg: ^1.0.2
+  flutter_svg: ^2.0.9
   intl: any
   logging: ^1.2.0
   path_provider: ^2.0.11
   provider: ^6.0.2
   streaming_shared_preferences: ^2.0.0
-  uuid: ^3.0.0
+  uuid: ^4.2.1
   archive: ^3.2.0
 
 dev_dependencies:
   flutter_test:
     sdk: flutter
-  flutter_lints: ^2.0.0
+  flutter_lints: ^3.0.1
 
 flutter:
   generate: true
diff --git a/Engine/Extras/VirtualProduction/EpicStageApp/lib/widgets/screens/main/tabs/web_browser_tab.dart b/Engine/Extras/VirtualProduction/EpicStageApp/lib/widgets/screens/main/tabs/web_browser_tab.dart
--- a/Engine/Extras/VirtualProduction/EpicStageApp/lib/widgets/screens/main/tabs/web_browser_tab.dart
+++ b/Engine/Extras/VirtualProduction/EpicStageApp/lib/widgets/screens/main/tabs/web_browser_tab.dart
@@ -39,10 +39,7 @@ class _WebBrowserTabState extends State<WebBrowserTab> {
   late final EngineConnectionManager _connectionManager;
 
   /// Controls the state of the web view.
-  WebViewController? _webViewController;
-
-  /// The initial URL to attempt to connect to.
-  String? _initialUrl;
+  WebViewController? _webViewController = null;
 
   /// The current URL to display in the address bar.
   String _currentUrl = '';
@@ -86,9 +83,13 @@ class _WebBrowserTabState extends State<WebBrowserTab> {
 
   @override
   Widget build(BuildContext context) {
+    if (_webViewController == null) {
+      _initWebViewController();
+    }
+
     late final Widget mainContent;
 
-    if (_errorMessage == null) {
+    if (_webViewController != null && _errorMessage == null) {
       // Show web view
       mainContent = Column(
         crossAxisAlignment: CrossAxisAlignment.stretch,
@@ -99,20 +100,11 @@ class _WebBrowserTabState extends State<WebBrowserTab> {
             child: Stack(
               children: [
                 if (!_bIsPageLoaded) const SpinnerOverlay(),
-                if (_initialUrl != null)
-                  Center(
-                    child: WebView(
-                      backgroundColor: Theme.of(context).colorScheme.surface,
-                      initialUrl: _initialUrl!,
-                      javascriptMode: JavascriptMode.unrestricted,
-                      allowsInlineMediaPlayback: true,
-                      onPageFinished: _onPageFinished,
-                      onWebViewCreated: (controller) => _webViewController = controller,
-                      onPageStarted: (String url) => setState(() {
-                        _currentUrl = url;
-                      }),
-                    ),
+                Center(
+                  child: WebViewWidget(
+                    controller: _webViewController!,
                   ),
+                ),
               ],
             ),
           ),
@@ -123,11 +115,13 @@ class _WebBrowserTabState extends State<WebBrowserTab> {
       mainContent = Center(
         child: EmptyPlaceholder(
           message: _errorMessage!,
-          button: EpicWideButton(
-            text: AppLocalizations.of(context)!.webBrowserReconnectButtonLabel,
-            iconPath: 'packages/epic_common/assets/icons/refresh.svg',
-            onPressed: _attemptToConnect,
-          ),
+          button: _webViewController != null
+              ? EpicWideButton(
+                  text: AppLocalizations.of(context)!.webBrowserReconnectButtonLabel,
+                  iconPath: 'packages/epic_common/assets/icons/refresh.svg',
+                  onPressed: _attemptToConnect,
+                )
+              : null,
         ),
       );
     }
@@ -141,26 +135,43 @@ class _WebBrowserTabState extends State<WebBrowserTab> {
   /// Refresh the currently loaded page, or attempt to reconnect if not currently showing a page.
   void refresh() {
     if (_webViewController == null) {
-      _attemptToConnect();
       return;
     }
 
+    _attemptToConnect();
     _webViewController!.reload();
   }
 
+  /// Try to set up the web view controller and set [_bIsWebViewAvailable] accordingly.
+  void _initWebViewController() {
+    if (WebViewPlatform.instance == null) {
+      // Not supported on this platform
+      _errorMessage = AppLocalizations.of(context)!.webBrowserUnsupportedMessage;
+      return;
+    }
+
+    _webViewController = WebViewController()
+      ..setBackgroundColor(Theme.of(context).colorScheme.surface)
+      ..setJavaScriptMode(JavaScriptMode.unrestricted)
+      ..setNavigationDelegate(
+        NavigationDelegate(
+          onPageFinished: _onPageFinished,
+          onPageStarted: (String url) => setState(() {
+            _currentUrl = url;
+          }),
+        ),
+      );
+  }
+
   /// Try the full connection process from scratch, including retrieving the URI from the engine.
   void _attemptToConnect() {
-    _initialUrl = null;
+    if (_webViewController == null) {
+      return;
+    }
+
     _errorMessage = null;
     _bIsPageLoaded = false;
 
-    try {
-      WebView.platform;
-    } catch (e) {
-      _errorMessage = AppLocalizations.of(context)!.webBrowserUnsupportedMessage;
-      setState(() {});
-    }
-
     if (_errorMessage == null) {
       _findBrowserURI().then((bool bSuccess) {
         if (!bSuccess) {
@@ -225,13 +236,13 @@ class _WebBrowserTabState extends State<WebBrowserTab> {
       }
 
       setState(() {
-        _initialUrl = 'http://${connectionData.websocketAddress.address}:$port';
+        _webViewController!.loadRequest(Uri.http('${connectionData.websocketAddress.address}:$port'));
       });
 
       return true;
     }).onError((error, stackTrace) {
       setState(() {
-        _log.warning('Failed to send request for web interface port.');
+        _log.warning('Failed to send request for web interface port.', error, stackTrace);
       });
       return false;
     });
diff --git a/Engine/Extras/VirtualProduction/EpicStageApp/lib/widgets/screens/main/toolbar/main_screen_tab_bar.dart b/Engine/Extras/VirtualProduction/EpicStageApp/lib/widgets/screens/main/toolbar/main_screen_tab_bar.dart
--- a/Engine/Extras/VirtualProduction/EpicStageApp/lib/widgets/screens/main/toolbar/main_screen_tab_bar.dart
+++ b/Engine/Extras/VirtualProduction/EpicStageApp/lib/widgets/screens/main/toolbar/main_screen_tab_bar.dart
@@ -33,7 +33,9 @@ class _MainScreenTabBarState extends State<MainScreenTabBar> {
         color: Theme.of(context).colorScheme.primary,
       ),
       indicatorSize: TabBarIndicatorSize.label,
+      labelColor: Colors.white,
       tabs: MainScreenTabs.tabConfigs.map((tabConfig) => _MainScreenTabBarTab(tabConfig: tabConfig)).toList(),
+      tabAlignment: TabAlignment.start,
     );
   }
 }
diff --git a/Engine/Extras/VirtualProduction/EpicStageApp/pubspec.yaml b/Engine/Extras/VirtualProduction/EpicStageApp/pubspec.yaml
--- a/Engine/Extras/VirtualProduction/EpicStageApp/pubspec.yaml
+++ b/Engine/Extras/VirtualProduction/EpicStageApp/pubspec.yaml
@@ -19,7 +19,7 @@ version: 1.0.2+1
 
 environment:
   sdk: ">=2.17.0 <3.0.0"
-  flutter: "^3.7.0"
+  flutter: "^3.16.3"
 
 # Dependencies specify other packages that your package needs in order to work.
 # To automatically upgrade your package dependencies to the latest versions
@@ -34,23 +34,23 @@ dependencies:
   epic_common:
     path: ../../FlutterCommon
 
-  uuid: ^3.0.0
+  uuid: ^4.2.1
   archive: ^3.2.0
   web_socket_channel: ^2.1.0
   provider: ^6.0.2
-  flutter_svg: ^1.0.2
+  flutter_svg: ^2.0.9
   streaming_shared_preferences: ^2.0.0
-  webview_flutter: ^3.0.4
+  webview_flutter: ^4.4.2
   logging: ^1.1.0
   intl: any
   path_provider: ^2.0.11
   open_filex: ^4.3.1
-  share_plus: ^6.3.0
+  share_plus: ^7.2.1
   flutter_localizations:
     sdk: flutter
   arb_generator: ^0.0.3
   xml: ^6.2.2
-  flutter_secure_storage: ^8.0.0
+  flutter_secure_storage: ^9.0.0
   url_launcher: ^6.1.11
 
 dev_dependencies:
@@ -62,9 +62,9 @@ dev_dependencies:
   # activated in the `analysis_options.yaml` file located at the root of your
   # package. See that file for information about deactivating specific lint
   # rules and activating additional ones.
-  flutter_lints: ^1.0.0
+  flutter_lints: ^3.0.1
 
-  flutter_launcher_icons: ^0.9.2
+  flutter_launcher_icons: ^0.13.1
 
 # For information on the generic Dart part of this file, see the
 # following page: https://dart.dev/tools/pub/pubspec
diff --git a/Engine/Extras/VirtualProduction/EpicStageApp/windows/flutter/CMakeLists.txt b/Engine/Extras/VirtualProduction/EpicStageApp/windows/flutter/CMakeLists.txt
--- a/Engine/Extras/VirtualProduction/EpicStageApp/windows/flutter/CMakeLists.txt
+++ b/Engine/Extras/VirtualProduction/EpicStageApp/windows/flutter/CMakeLists.txt
@@ -9,6 +9,11 @@ include(${EPHEMERAL_DIR}/generated_config.cmake)
 # [REDACTED_PROFILE_URL]/flutter/issues/57146.
 set(WRAPPER_ROOT "${EPHEMERAL_DIR}/cpp_client_wrapper")
 
+# Set fallback configurations for older versions of the flutter tool.
+if (NOT DEFINED FLUTTER_TARGET_PLATFORM)
+  set(FLUTTER_TARGET_PLATFORM "windows-x64")
+endif()
+
 # === Flutter Library ===
 set(FLUTTER_LIBRARY "${EPHEMERAL_DIR}/flutter_windows.dll")
 
@@ -91,7 +96,7 @@ add_custom_command(
   COMMAND ${CMAKE_COMMAND} -E env
     ${FLUTTER_TOOL_ENVIRONMENT}
     "${FLUTTER_ROOT}/packages/flutter_tools/bin/tool_backend.bat"
-      windows-x64 $<CONFIG>
+      ${FLUTTER_TARGET_PLATFORM} $<CONFIG>
   VERBATIM
 )
 add_custom_target(flutter_assemble DEPENDS
diff --git a/Engine/Extras/VirtualProduction/FlutterLiveLinkVCAM/android/app/src/main/AndroidManifest.xml b/Engine/Extras/VirtualProduction/FlutterLiveLinkVCAM/android/app/src/main/AndroidManifest.xml
--- a/Engine/Extras/VirtualProduction/FlutterLiveLinkVCAM/android/app/src/main/AndroidManifest.xml
+++ b/Engine/Extras/VirtualProduction/FlutterLiveLinkVCAM/android/app/src/main/AndroidManifest.xml
@@ -14,7 +14,8 @@
     <application
         android:label="live_link_vcam"
         android:name="${applicationName}"
-        android:icon="@mipmap/ic_launcher">
+        android:icon="@mipmap/ic_launcher"
+        android:enableOnBackInvokedCallback="true">
         <activity
             android:name=".MainActivity"
             android:exported="true"
diff --git a/Engine/Extras/VirtualProduction/FlutterLiveLinkVCAM/pubspec.yaml b/Engine/Extras/VirtualProduction/FlutterLiveLinkVCAM/pubspec.yaml
--- a/Engine/Extras/VirtualProduction/FlutterLiveLinkVCAM/pubspec.yaml
+++ b/Engine/Extras/VirtualProduction/FlutterLiveLinkVCAM/pubspec.yaml
@@ -6,7 +6,7 @@ version: 1.0.0+1
 
 environment:
   sdk: '>=3.0.6 <4.0.0'
-  flutter: "^3.7.0"
+  flutter: "^3.16.3"
 
 dependencies:
   flutter:
@@ -16,14 +16,14 @@ dependencies:
     path: ../../FlutterCommon
 
   logging: ^1.2.0
-  package_info_plus: ^4.1.0
+  package_info_plus: ^5.0.1
   web_socket_channel: ^2.1.0
 
 dev_dependencies:
   flutter_test:
     sdk: flutter
 
-  pigeon: ^11.0.1
+  pigeon: ^14.0.1
   share_plus: ^7.2.1
 
   # The "flutter_lints" package below contains a set of recommended lints to
@@ -31,7 +31,7 @@ dev_dependencies:
   # activated in the `analysis_options.yaml` file located at the root of your
   # package. See that file for information about deactivating specific lint
   # rules and activating additional ones.
-  flutter_lints: ^2.0.0
+  flutter_lints: ^3.0.1
 
 # For information on the generic Dart part of this file, see the
 # following page: https://dart.dev/tools/pub/pubspec
