[UnrealStage] Demo mode

diff --git a/Engine/Extras/VirtualProduction/EpicStageApp/lib/models/engine_connection.dart b/Engine/Extras/VirtualProduction/EpicStageApp/lib/models/engine_connection.dart
--- a/Engine/Extras/VirtualProduction/EpicStageApp/lib/models/engine_connection.dart
+++ b/Engine/Extras/VirtualProduction/EpicStageApp/lib/models/engine_connection.dart
@@ -14,13 +14,13 @@ import 'package:provider/provider.dart';
 import 'package:shared_preferences/shared_preferences.dart';
 import 'package:web_socket_channel/io.dart';
 
-import './navigator_keys.dart';
 import '../utilities/net_utilities.dart';
 import '../utilities/unreal_utilities.dart';
 import '../widgets/screens/connect/connect.dart';
 import '../widgets/screens/connect/views/passphrase_modal.dart';
 import '../widgets/screens/main/stage_app_main_screen.dart';
 import '../widgets/screens/reconnect_screen.dart';
+import './navigator_keys.dart';
 import 'api_version.dart';
 import 'engine_passphrase_manager.dart';
 import 'settings/connection_settings.dart';
@@ -165,6 +165,9 @@ class EngineConnectionManager with WidgetsBindingObserver {
   /// Whether this should attempt to reconnect when the app becomes active again.
   bool _bShouldReconnectOnWake = false;
 
+  /// Whether the current connection is in demo mode, meaning no actual Unreal Engine instance is connected.
+  bool _bIsInDemoMode = false;
+
   /// MD5 hash of the passphrase to be passed along with all messages for the current connection, or null if no
   /// passphrase was given.
   String? _passphraseHash;
@@ -187,6 +190,9 @@ class EngineConnectionManager with WidgetsBindingObserver {
   /// Get the version of the engine we're currently connected to, or null if we're not connected.
   EpicStageAppAPIVersion? get apiVersion => _apiVersion;
 
+  /// Whether the current connection is in demo mode, meaning no actual Unreal Engine instance is connected.
+  bool get bIsInDemoMode => _bIsInDemoMode;
+
   /// Dispose of any stored data.
   void dispose() {
     disconnect();
@@ -199,6 +205,7 @@ class EngineConnectionManager with WidgetsBindingObserver {
       case AppLifecycleState.detached:
       case AppLifecycleState.paused:
       case AppLifecycleState.inactive:
+      case AppLifecycleState.hidden:
         _onAppPaused();
         break;
 
@@ -275,6 +282,8 @@ class EngineConnectionManager with WidgetsBindingObserver {
       // the navigator by itself
       bShouldReturnToConnectScreen: _pendingConnectionAttempt == null,
     );
+
+    _bIsInDemoMode = false;
   }
 
   /// Send a message via the current WebSocket connection in standard Unreal WebSocket format.
@@ -304,6 +313,10 @@ class EngineConnectionManager with WidgetsBindingObserver {
   /// Send a raw message via the current WebSocket connection, encoded as JSON.
   /// You may want to create the message using [createUnrealWebSocketMessage].
   void sendRawMessage(dynamic message) {
+    if (_bIsInDemoMode) {
+      return;
+    }
+
     if (connectionState != EngineConnectionState.connected && _pendingConnectionAttempt == null) {
       return;
     }
@@ -459,6 +472,15 @@ class EngineConnectionManager with WidgetsBindingObserver {
   /// Internal function to connect to the engine.
   /// You should almost always call [connect] instead of this so we know we have a pending connection attempt.
   Future<EngineConnectionResult> _internalConnect(ConnectionData connectionData) async {
+    _bIsInDemoMode = connectionData.bIsDemo;
+
+    if (_bIsInDemoMode) {
+      // Run the app in demo mode, where the connection will be spoofed
+      _connectionState = EngineConnectionState.connected;
+      _apiVersion = EpicStageAppAPIVersion(0, 0, 0);
+      return Future.value(EngineConnectionResult.success);
+    }
+
     // Get the passphrase ready
     final passphraseManager = Provider.of<EnginePassphraseManager>(context, listen: false);
     final String? passphrase = await passphraseManager.getPassphrase(connectionData);
diff --git a/Engine/Extras/VirtualProduction/EpicStageApp/lib/models/preview_render_manager.dart b/Engine/Extras/VirtualProduction/EpicStageApp/lib/models/preview_render_manager.dart
--- a/Engine/Extras/VirtualProduction/EpicStageApp/lib/models/preview_render_manager.dart
+++ b/Engine/Extras/VirtualProduction/EpicStageApp/lib/models/preview_render_manager.dart
@@ -335,6 +335,11 @@ class PreviewRenderManager with WidgetsBindingObserver {
 
   /// Send a request to initialize the preview renderer in the engine.
   void _initializePreviewRenderer() {
+    if (_connectionManager.bIsInDemoMode) {
+      _onPreviewRendererCreated({'RendererId': 0});
+      return;
+    }
+
     final rootPath = _selectedActorSettings.displayClusterRootPath.getValue();
     if (rootPath.isEmpty) {
       return;
@@ -408,6 +413,15 @@ class PreviewRenderManager with WidgetsBindingObserver {
       return;
     }
 
+    if (_connectionManager.bIsInDemoMode) {
+      _onPreviewRenderCompleted({
+        'RendererId': _previewRendererId,
+        // A 1x1 black square encoded as a base64 image so that we have something to display
+        'ImageBase64': 'iVBORw0KGgoAAAANSUhEUgAAAAEAAAABCAQAAAC1HAwCAAAAC0lEQVR42mNk+A8AAQUBAScY42YAAAAASUVORK5CYII=',
+      });
+      return;
+    }
+
     _requestRenderHelper = MessageTimeoutRetryHelper(
       logger: _log,
       logDescription: 'new preview render',
@@ -471,27 +485,30 @@ class PreviewRenderManager with WidgetsBindingObserver {
       // Gather actor position data
       final List<PreviewActorData> actorPositions = [];
       final dynamic actorPositionsData = message['ActorPositions'];
-      for (dynamic positionPair in actorPositionsData) {
-        final String? actorPath = positionPair['Path'];
-        if (actorPath == null) {
-          continue;
-        }
 
-        final dynamic actorPositionData = positionPair['Position'];
-        if (actorPositionData == null) {
-          continue;
+      if (actorPositionsData is List) {
+        for (dynamic positionPair in actorPositionsData) {
+          final String? actorPath = positionPair['Path'];
+          if (actorPath == null) {
+            continue;
+          }
+
+          final dynamic actorPositionData = positionPair['Position'];
+          if (actorPositionData == null) {
+            continue;
+          }
+
+          final double? xPosition = jsonNumberToDouble(actorPositionData['X']);
+          final double? yPosition = jsonNumberToDouble(actorPositionData['Y']);
+          if (xPosition == null || yPosition == null) {
+            continue;
+          }
+
+          actorPositions.add(PreviewActorData(
+            path: actorPath,
+            position: Offset(xPosition, yPosition),
+          ));
         }
-
-        final double? xPosition = jsonNumberToDouble(actorPositionData['X']);
-        final double? yPosition = jsonNumberToDouble(actorPositionData['Y']);
-        if (xPosition == null || yPosition == null) {
-          continue;
-        }
-
-        actorPositions.add(PreviewActorData(
-          path: actorPath,
-          position: Offset(xPosition, yPosition),
-        ));
       }
 
       for (final PreviewRenderConsumer consumer in _consumers) {
diff --git a/Engine/Extras/VirtualProduction/EpicStageApp/lib/models/unreal_actor_manager.dart b/Engine/Extras/VirtualProduction/EpicStageApp/lib/models/unreal_actor_manager.dart
--- a/Engine/Extras/VirtualProduction/EpicStageApp/lib/models/unreal_actor_manager.dart
+++ b/Engine/Extras/VirtualProduction/EpicStageApp/lib/models/unreal_actor_manager.dart
@@ -123,6 +123,10 @@ class UnrealActorManager with WidgetsBindingObserver {
   /// Start watching for any engine-side changes to actors of the Unreal type [className], and call [callback] when
   /// there's any change in the list of relevant actors.
   void watchClassName(String className, ActorUpdateCallback callback) {
+    if (_connectionManager.bIsInDemoMode) {
+      return;
+    }
+
     if (_classWatchCallbacks.containsKey(className)) {
       // We're already watching this class, so just add the callback
       _classWatchCallbacks[className]!.add(callback);
diff --git a/Engine/Extras/VirtualProduction/EpicStageApp/lib/models/unreal_property_manager.dart b/Engine/Extras/VirtualProduction/EpicStageApp/lib/models/unreal_property_manager.dart
--- a/Engine/Extras/VirtualProduction/EpicStageApp/lib/models/unreal_property_manager.dart
+++ b/Engine/Extras/VirtualProduction/EpicStageApp/lib/models/unreal_property_manager.dart
@@ -6,10 +6,10 @@ import 'package:flutter/material.dart';
 import 'package:logging/logging.dart';
 import 'package:provider/provider.dart';
 
+import '../utilities/unreal_utilities.dart';
 import './engine_connection.dart';
 import './property_modify_operations.dart';
 import './unreal_types.dart';
-import '../utilities/unreal_utilities.dart';
 import 'unreal_transaction_manager.dart';
 
 final _log = Logger('UnrealPropertyManager');
@@ -509,6 +509,10 @@ class UnrealPropertyManager {
 
   /// Called each tick, when we're ready to send updates back to the engine.
   void _onTick(Timer timer) {
+    if (_connectionManager.bIsInDemoMode) {
+      return;
+    }
+
     final List<dynamic> messages = _finalizeTickChanges();
 
     if (_transientPresetName == null) {
diff --git a/Engine/Extras/VirtualProduction/EpicStageApp/lib/utilities/net_utilities.dart b/Engine/Extras/VirtualProduction/EpicStageApp/lib/utilities/net_utilities.dart
--- a/Engine/Extras/VirtualProduction/EpicStageApp/lib/utilities/net_utilities.dart
+++ b/Engine/Extras/VirtualProduction/EpicStageApp/lib/utilities/net_utilities.dart
@@ -43,7 +43,28 @@ class UnrealStageBeaconData {
 
 /// Data about a connection to Unreal Engine.
 class ConnectionData {
-  const ConnectionData({required this.name, required this.websocketAddress, required this.websocketPort});
+  const ConnectionData({
+    required this.name,
+    required this.websocketAddress,
+    required this.websocketPort,
+  }) : this.bIsDemo = false;
+
+  const ConnectionData._({
+    required this.name,
+    required this.websocketAddress,
+    required this.websocketPort,
+    required this.bIsDemo,
+  });
+
+  /// Create connection data for demo mode.
+  static ConnectionData forDemoMode() {
+    return ConnectionData._(
+      name: 'Demo',
+      websocketAddress: InternetAddress('127.0.0.1'),
+      websocketPort: 0,
+      bIsDemo: true,
+    );
+  }
 
   /// Create connection data from a [response] to an UnrealEngineBeacon message.
   static ConnectionData fromBeaconResponse(UnrealStageBeaconResponse response) {
@@ -62,6 +83,9 @@ class ConnectionData {
 
   /// Port to connect to for WebSocket communication.
   final int websocketPort;
+
+  /// Whether this connection is a demo, meaning no actual Unreal Engine instance exists.
+  final bool bIsDemo;
 }
 
 /// A helper class that automatically retries sending a message to the engine if it times out.
diff --git a/Engine/Extras/VirtualProduction/EpicStageApp/lib/widgets/elements/lightcard_map.dart b/Engine/Extras/VirtualProduction/EpicStageApp/lib/widgets/elements/lightcard_map.dart
--- a/Engine/Extras/VirtualProduction/EpicStageApp/lib/widgets/elements/lightcard_map.dart
+++ b/Engine/Extras/VirtualProduction/EpicStageApp/lib/widgets/elements/lightcard_map.dart
@@ -26,7 +26,6 @@ import '../../models/unreal_property_manager.dart';
 import '../../models/unreal_types.dart';
 import '../../utilities/constants.dart';
 import '../../utilities/guarded_refresh_state.dart';
-import 'dropdown_button.dart';
 import 'place_actor_menu.dart';
 import 'spinner_overlay.dart';
 import 'transform_gesture_detector.dart';
@@ -423,11 +422,12 @@ class StageMapState extends State<StageMap> with PreviewRenderConsumer, GuardedR
     _propertyManager = Provider.of<UnrealPropertyManager>(context, listen: false);
     _updateActorTransformRegistration();
 
+    _previewImageStreamListener = ImageStreamListener(_onPreviewImageLoaded, onError: _onPreviewImageError);
+
     _previewRenderManager = Provider.of<PreviewRenderManager>(context, listen: false);
     _previewRenderManager.addConsumer(this);
 
     _mapTransform.addListener(_saveMapTransformToUserSettings);
-    _previewImageStreamListener = ImageStreamListener(_onPreviewImageLoaded, onError: _onPreviewImageError);
   }
 
   @override
diff --git a/Engine/Extras/VirtualProduction/EpicStageApp/lib/widgets/screens/connect/views/manual_connect_form.dart b/Engine/Extras/VirtualProduction/EpicStageApp/lib/widgets/screens/connect/views/manual_connect_form.dart
--- a/Engine/Extras/VirtualProduction/EpicStageApp/lib/widgets/screens/connect/views/manual_connect_form.dart
+++ b/Engine/Extras/VirtualProduction/EpicStageApp/lib/widgets/screens/connect/views/manual_connect_form.dart
@@ -10,6 +10,9 @@ import 'package:flutter_gen/gen_l10n/app_localizations.dart';
 import '../../../../utilities/constants.dart';
 import '../../../../utilities/net_utilities.dart';
 
+/// String which will trigger demo mode when entered as the IP address.
+const String _demoModeString = 'demo.mode';
+
 /// Manual connection form for manually connecting to an instance of UE.
 class ManualConnectForm extends StatefulWidget {
   const ManualConnectForm({Key? key, required this.connect}) : super(key: key);
@@ -101,11 +104,17 @@ class _ManualConnectFormState extends State<ManualConnectForm> {
 
     Navigator.of(context).pop(false);
 
-    ConnectionData connection = ConnectionData(
-      name: AppLocalizations.of(context)!.connectScreenManualConnectionLabel,
-      websocketAddress: InternetAddress(_ipTextController.text),
-      websocketPort: int.parse(_portTextController.text),
-    );
+    final ConnectionData connection;
+
+    if (_ipTextController.text == _demoModeString) {
+      connection = ConnectionData.forDemoMode();
+    } else {
+      connection = ConnectionData(
+        name: AppLocalizations.of(context)!.connectScreenManualConnectionLabel,
+        websocketAddress: InternetAddress(_ipTextController.text),
+        websocketPort: int.parse(_portTextController.text),
+      );
+    }
 
     /// Form input validating has passed, we are calling [widget.connect] to initiate a connection with UE using
     /// [connection]
@@ -114,6 +123,10 @@ class _ManualConnectFormState extends State<ManualConnectForm> {
 
   /// String validation for validating IpAddresses.
   String? _validateAddress(String? text) {
+    if (text == _demoModeString) {
+      return null;
+    }
+
     if (text == null) {
       return AppLocalizations.of(context)!.formErrorInvalidIPAddress;
     }
