[FlutterCommon] [UnrealStage] Move missing fragment shader file from Unreal Stage to FlutterCommon

diff --git a/Engine/Extras/FlutterCommon/pubspec.yaml b/Engine/Extras/FlutterCommon/pubspec.yaml
--- a/Engine/Extras/FlutterCommon/pubspec.yaml
+++ b/Engine/Extras/FlutterCommon/pubspec.yaml
@@ -34,6 +34,9 @@ flutter:
     - assets/icons/
     - assets/images/
 
+  shaders:
+    - assets/shaders/scroll_view_fade.frag
+
 # Automatically generate arb localization files from a CSV
 arb_generator:
   input_filepath: "lib/src/l10n/l10n.csv"
diff --git a/Engine/Extras/VirtualProduction/EpicStageApp/lib/main.dart b/Engine/Extras/VirtualProduction/EpicStageApp/lib/main.dart
--- a/Engine/Extras/VirtualProduction/EpicStageApp/lib/main.dart
+++ b/Engine/Extras/VirtualProduction/EpicStageApp/lib/main.dart
@@ -80,11 +80,6 @@ Future preloadShaders() async {
   await EpicCommonWidgets.preloadShaders();
 }
 
-/// Unload any preloaded shaders.
-void unloadShaders() {
-  EpicScrollView.unloadShaders();
-}
-
 class EpicStageApp extends StatefulWidget {
   const EpicStageApp({Key? key, required this.preferences}) : super(key: key);
 
diff --git a/Engine/Extras/VirtualProduction/EpicStageApp/pubspec.yaml b/Engine/Extras/VirtualProduction/EpicStageApp/pubspec.yaml
--- a/Engine/Extras/VirtualProduction/EpicStageApp/pubspec.yaml
+++ b/Engine/Extras/VirtualProduction/EpicStageApp/pubspec.yaml
@@ -82,9 +82,6 @@ flutter:
 
   assets:
     - assets/images/map/
-  
-  shaders:
-    - assets/shaders/scroll_view_fade.frag
 
   fonts:
     - family: Inter
