[LiveLinkVCAM] Implement timecode via system clock or NTP pool

diff --git a/Engine/Extras/VirtualProduction/FlutterLiveLinkVCAM/pubspec.yaml b/Engine/Extras/VirtualProduction/FlutterLiveLinkVCAM/pubspec.yaml
--- a/Engine/Extras/VirtualProduction/FlutterLiveLinkVCAM/pubspec.yaml
+++ b/Engine/Extras/VirtualProduction/FlutterLiveLinkVCAM/pubspec.yaml
@@ -16,7 +16,9 @@ dependencies:
     path: ../../FlutterCommon
 
   logging: ^1.2.0
+  ntp: ^2.0.0
   package_info_plus: ^5.0.1
+  provider: ^6.1.1
   web_socket_channel: ^2.1.0
 
 dev_dependencies:
@@ -43,6 +45,9 @@ flutter:
   # Generate localization files
   generate: true
 
+  assets:
+    - assets/icons/
+
   fonts:
     - family: Inter
       fonts:
