[LiveLinkVCAM] Tentacle Sync timecode support

diff --git a/Engine/Extras/VirtualProduction/FlutterLiveLinkVCAM/IOS/Runner/AppDelegate.swift b/Engine/Extras/VirtualProduction/FlutterLiveLinkVCAM/IOS/Runner/AppDelegate.swift
--- a/Engine/Extras/VirtualProduction/FlutterLiveLinkVCAM/IOS/Runner/AppDelegate.swift
+++ b/Engine/Extras/VirtualProduction/FlutterLiveLinkVCAM/IOS/Runner/AppDelegate.swift
@@ -10,6 +10,7 @@ import WebRTC
   var dataChannelApi: FlutterRtcDataChannelApi?
   var videoViewControllerApi: FlutterRtcVideoViewControllerApi?
   var arSessionApi: FlutterArSessionApi?
+  var tentacleApi: FlutterTentacleApi?
   
   override func application(
     _ application: UIApplication,
@@ -20,7 +21,7 @@ import WebRTC
     GeneratedPluginRegistrant.register(with: self)
     
     // Register RtcVideoView platform view
-    var videoViewRegistrar: FlutterPluginRegistrar? = self.registrar(forPlugin: "RtcVideoView")
+    let videoViewRegistrar: FlutterPluginRegistrar? = self.registrar(forPlugin: "RtcVideoView")
     if (videoViewRegistrar == nil) {
       fatalError("Failed to register RtcVideoView plugin")
     }
@@ -39,6 +40,7 @@ import WebRTC
     dataChannelApi = FlutterRtcDataChannelApi(binaryMessenger: binaryMessenger)
     videoViewControllerApi = FlutterRtcVideoViewControllerApi(binaryMessenger: binaryMessenger)
     arSessionApi = FlutterArSessionApi(binaryMessenger: binaryMessenger)
+    tentacleApi = FlutterTentacleApi(binaryMessenger: binaryMessenger)
   
     return super.application(application, didFinishLaunchingWithOptions: launchOptions)
   }
diff --git a/Engine/Extras/VirtualProduction/FlutterLiveLinkVCAM/IOS/Runner/Tentacle/API/FlutterTentacleApi.swift b/Engine/Extras/VirtualProduction/FlutterLiveLinkVCAM/IOS/Runner/Tentacle/API/FlutterTentacleApi.swift
new file mode 100644
--- /dev/null
+++ b/Engine/Extras/VirtualProduction/FlutterLiveLinkVCAM/IOS/Runner/Tentacle/API/FlutterTentacleApi.swift
@@ -0,0 +1,90 @@
+// Copyright Epic Games, Inc. All Rights Reserved.
+
+import ARKit
+import Foundation
+import Tentacle
+
+/// Handles messages about Tentacle to and from the Flutter API.
+class FlutterTentacleApi: FlutterPluginApi<TentacleFlutterApi> {
+  /// The Bluetooth controller that scans for Tentacle devices.
+  private var bluetooth: TentacleBluetoothController!
+  
+  /// Whether this should attempt to scan for Tentacle devices.
+  private var bIsScanning: Bool = false
+  
+  init(binaryMessenger: FlutterBinaryMessenger) {
+    super.init(flutter:TentacleFlutterApi(binaryMessenger: binaryMessenger))
+    
+    TentacleHostApiSetup.setUp(binaryMessenger: binaryMessenger, api: self)
+    FlutterTentacleApi.instance = self
+    
+    bluetooth = TentacleBluetoothController()
+    bluetooth.delegate = self
+  }
+  
+  deinit {
+    if (FlutterArSessionApi.instance === self) {
+      FlutterTentacleApi.instance = nil
+    }
+  }
+  
+  /// The active singleton instance of this class.
+  private(set) static var instance: FlutterTentacleApi?
+  
+  /// Convert a C string with the given length to a Swift string.
+  private func stringFromCString<T>(cString: inout T, length: Int) -> String {
+    return withUnsafePointer(to: cString) { startPointer in
+      return startPointer.withMemoryRebound(to: UInt8.self, capacity: MemoryLayout.size(ofValue: length)) { cStringPointer in
+          return String(cString: cStringPointer)
+      }
+    }
+  }
+}
+
+// MARK: Host API
+extension FlutterTentacleApi: TentacleHostApi {
+  func startScanning() throws {
+    bIsScanning = true
+    bluetooth.startScanning()
+  }
+  
+  func stopScanning() throws {
+    bIsScanning = false
+    bluetooth.stopScanning()
+  }
+  
+  func getCachedDeviceInfo() throws -> [TentacleDeviceInfo] {
+    var devices: [TentacleDeviceInfo] = []
+    let deviceCount: Int32 = Int32(Tentacle.TentacleDeviceCacheGetSize())
+  
+    for deviceIndex in stride(from: 0, to: deviceCount, by: 1) {
+      let device: TentacleDevice = Tentacle.TentacleDeviceCacheGetDevice(deviceIndex)
+      var advert: TentacleAdvertisement = device.advertisement
+      
+      devices.append(TentacleDeviceInfo(
+        bIsInGreenMode: advert.greenMode,
+        bIsCharging: advert.charging,
+        bIsDropFrame: advert.dropFrame,
+        batteryLevel: Int64(advert.battery),
+        frameRate: Tentacle.TentacleAdvertisementGetFrameRate(&advert),
+        lastSeenTimestamp: CACurrentMediaTime(),
+        iconIndex: Int64(advert.icon),
+        signalStrength: Int64(advert.rssi),
+        productId: advert.productId.toFlutter(),
+        identifier: stringFromCString(cString: &advert.identifier, length: advert.identifierLength),
+        name: stringFromCString(cString: &advert.name, length: advert.nameLength)
+      ))
+    }
+    
+    return devices
+  }
+}
+
+// MARK: Bluetooth Controller Delegate
+extension FlutterTentacleApi : TentacleBluetoothControllerDelegate {
+  func didUpdate(to state: TentacleBluetoothState) {
+      if state == .poweredOn && bIsScanning {
+          bluetooth.startScanning()
+      }
+  }
+}
diff --git a/Engine/Extras/VirtualProduction/FlutterLiveLinkVCAM/IOS/Runner/Tentacle/Util/TentacleExtensions.swift b/Engine/Extras/VirtualProduction/FlutterLiveLinkVCAM/IOS/Runner/Tentacle/Util/TentacleExtensions.swift
new file mode 100644
--- /dev/null
+++ b/Engine/Extras/VirtualProduction/FlutterLiveLinkVCAM/IOS/Runner/Tentacle/Util/TentacleExtensions.swift
@@ -0,0 +1,14 @@
+// Copyright Epic Games, Inc. All Rights Reserved.
+
+import Tentacle
+
+extension Tentacle.TentacleProductId {
+  /// Convert from a native Tentacle product ID to the Flutter equivalent.
+  func toFlutter() -> TentacleProductId {
+    switch (self) {
+    case TentacleProductIdSyncE: return .syncE
+    case TentacleProductIdTrackE: return .trackE
+    default: return .generic
+    }
+  }
+}
diff --git a/Engine/Extras/VirtualProduction/FlutterLiveLinkVCAM/IOS/Runner/Util/FlutterPluginApi.swift b/Engine/Extras/VirtualProduction/FlutterLiveLinkVCAM/IOS/Runner/Util/FlutterPluginApi.swift
--- a/Engine/Extras/VirtualProduction/FlutterLiveLinkVCAM/IOS/Runner/Util/FlutterPluginApi.swift
+++ b/Engine/Extras/VirtualProduction/FlutterLiveLinkVCAM/IOS/Runner/Util/FlutterPluginApi.swift
@@ -3,7 +3,7 @@
 import Foundation
 
 /// Base class for all APIs that communicate to and from a Flutter plugin.
-class FlutterPluginApi<T> {
+class FlutterPluginApi<T>: NSObject {
   /// The API used to send messages to Flutter.
   private let flutter: T
   
diff --git a/Engine/Extras/VirtualProduction/FlutterLiveLinkVCAM/README.md b/Engine/Extras/VirtualProduction/FlutterLiveLinkVCAM/README.md
--- a/Engine/Extras/VirtualProduction/FlutterLiveLinkVCAM/README.md
+++ b/Engine/Extras/VirtualProduction/FlutterLiveLinkVCAM/README.md
@@ -2,18 +2,22 @@
 
 A rebuild of the Live Link VCAM app for Flutter, enabling cross-platform deployment.
 
-## Getting Started
+## Tentacle Integration
 
-This project is a starting point for a Flutter application.
+This project depends on the Tentacle SDK located in Engine/Restricted/NotForLicensees/Source/ThirdParty/TentacleSDK.
 
-A few resources to get you started if this is your first Flutter project:
+We use some Tentacle SDK functions through Dart's Foreign Function Interface (FFI). Dart bindings for these functions are
+automatically generated using ffigen and checked into source control.
 
-- [Lab: Write your first Flutter app](https://docs.flutter.dev/get-started/codelab)
-- [Cookbook: Useful Flutter samples](https://docs.flutter.dev/cookbook)
+If you need to regenerate them, first set up ffigen (see https://pub.dev/packages/ffigen), then run the following command
+in the root VCAM project directory:
 
-For help getting started with Flutter development, view the
-[online documentation](https://docs.flutter.dev/), which offers tutorials,
-samples, guidance on mobile development, and a full API reference.
+    `dart run ffigen --config ffigen_tentacle.yaml`
+
+## Pigeon
+
+This project uses Pigeon to automatically generate platform channel bindings for native APIs.
+See `pigeons/README` for more information.
 
 ## Signing Certificates 
 
@@ -22,5 +26,4 @@ To setup signing certificates, follow the instructions outlined in this article
  - [Deploying Flutter Apps to The PlayStore](https:[REDACTED_EMAIL]/deploying-flutter-apps-to-the-playstore-1bd0cce0d15c)
 
 Add your `android-key.tks` to `android\app` folder and `key.properities` to `android\`.  The gradle files are already setup
-to read these files when app bundling
-
+to read these files when app bundling
\ No newline at end of file
diff --git a/Engine/Extras/VirtualProduction/FlutterLiveLinkVCAM/android/app/src/main/AndroidManifest.xml b/Engine/Extras/VirtualProduction/FlutterLiveLinkVCAM/android/app/src/main/AndroidManifest.xml
--- a/Engine/Extras/VirtualProduction/FlutterLiveLinkVCAM/android/app/src/main/AndroidManifest.xml
+++ b/Engine/Extras/VirtualProduction/FlutterLiveLinkVCAM/android/app/src/main/AndroidManifest.xml
@@ -1,8 +1,12 @@
 <manifest xmlns:android="http://schemas.android.com/apk/res/android"
-	package="com.epicgames.live_link_vcam">
+    xmlns:tools="http://schemas.android.com/tools"
+    package="com.epicgames.live_link_vcam">
 
     <uses-feature android:name="android.hardware.camera" />
     <uses-feature android:name="android.hardware.camera.ar" />
+    <uses-feature
+        android:name="android.hardware.bluetooth_le"
+        android:required="false" />
 
     <uses-permission android:name="android.permission.CAMERA" />
     <uses-permission android:name="android.permission.RECORD_AUDIO" />
@@ -11,6 +15,23 @@
     <uses-permission android:name="android.permission.CHANGE_NETWORK_STATE" />
     <uses-permission android:name="android.permission.MODIFY_AUDIO_SETTINGS" />
 
+    <!-- Bluetooth permissions for Tentacle scanning -->
+    <!-- API <= 30 -->
+    <uses-permission
+        android:name="android.permission.BLUETOOTH"
+        android:maxSdkVersion="30" />
+    <uses-permission
+        android:name="android.permission.BLUETOOTH_ADMIN"
+        android:maxSdkVersion="30" />
+    <uses-permission android:name="android.permission.ACCESS_COARSE_LOCATION" />
+    <uses-permission android:name="android.permission.ACCESS_FINE_LOCATION" />
+    <!-- API >= 31 -->
+    <uses-permission
+        android:name="android.permission.BLUETOOTH_SCAN"
+        android:usesPermissionFlags="neverForLocation"
+        tools:targetApi="s" />
+    <uses-permission android:name="android.permission.BLUETOOTH_CONNECT" />
+
     <application
         android:label="live_link_vcam"
         android:name="${applicationName}"
diff --git a/Engine/Extras/VirtualProduction/FlutterLiveLinkVCAM/devtools_options.yaml b/Engine/Extras/VirtualProduction/FlutterLiveLinkVCAM/devtools_options.yaml
new file mode 100644
--- /dev/null
+++ b/Engine/Extras/VirtualProduction/FlutterLiveLinkVCAM/devtools_options.yaml
@@ -0,0 +1 @@
+extensions:
diff --git a/Engine/Extras/VirtualProduction/FlutterLiveLinkVCAM/pubspec.yaml b/Engine/Extras/VirtualProduction/FlutterLiveLinkVCAM/pubspec.yaml
--- a/Engine/Extras/VirtualProduction/FlutterLiveLinkVCAM/pubspec.yaml
+++ b/Engine/Extras/VirtualProduction/FlutterLiveLinkVCAM/pubspec.yaml
@@ -5,7 +5,7 @@ publish_to: 'none'
 version: 1.0.0+1
 
 environment:
-  sdk: '>=3.0.6 <4.0.0'
+  sdk: '>=3.1.0 <4.0.0'
   flutter: "^3.16.3"
 
 dependencies:
@@ -15,16 +15,21 @@ dependencies:
   epic_common:
     path: ../../FlutterCommon
 
+  device_info_plus: ^9.1.1
+  ffi: ^2.1.0
   logging: ^1.2.0
   ntp: ^2.0.0
   package_info_plus: ^5.0.1
+  permission_handler: ^11.1.0
   provider: ^6.1.1
+  system_clock: ^2.0.0
   web_socket_channel: ^2.1.0
 
 dev_dependencies:
   flutter_test:
     sdk: flutter
 
+  ffigen: ^10.0.0
   pigeon: ^14.0.1
   share_plus: ^7.2.1
 
