[LiveLinkVCAM] Basic settings UI

diff --git a/Engine/Extras/VirtualProduction/FlutterLiveLinkVCAM/pubspec.yaml b/Engine/Extras/VirtualProduction/FlutterLiveLinkVCAM/pubspec.yaml
--- a/Engine/Extras/VirtualProduction/FlutterLiveLinkVCAM/pubspec.yaml
+++ b/Engine/Extras/VirtualProduction/FlutterLiveLinkVCAM/pubspec.yaml
@@ -12,6 +12,9 @@ dependencies:
   flutter:
     sdk: flutter
 
+  flutter_localizations:
+    sdk: flutter
+
   epic_common:
     path: ../../FlutterCommon
 
@@ -24,6 +27,8 @@ dependencies:
   provider: ^6.1.1
   system_clock: ^2.0.0
   web_socket_channel: ^2.1.0
+  streaming_shared_preferences: ^2.0.0
+  share_plus: ^7.2.1
 
 dev_dependencies:
   flutter_test:
@@ -31,7 +36,7 @@ dev_dependencies:
 
   ffigen: ^10.0.0
   pigeon: ^14.0.1
-  share_plus: ^7.2.1
+  url_launcher: ^6.2.2
 
   # The "flutter_lints" package below contains a set of recommended lints to
   # encourage good coding practices. The lint set provided by the package is
