[UnrealStage] Actor lat/long control via trackpad

diff --git a/Engine/Extras/VirtualProduction/EpicStageApp/lib/main.dart b/Engine/Extras/VirtualProduction/EpicStageApp/lib/main.dart
--- a/Engine/Extras/VirtualProduction/EpicStageApp/lib/main.dart
+++ b/Engine/Extras/VirtualProduction/EpicStageApp/lib/main.dart
@@ -27,6 +27,7 @@ import 'models/settings/selected_actor_settings.dart';
 import 'models/settings/stage_map_settings.dart';
 import 'models/unreal_actor_creator.dart';
 import 'models/unreal_actor_manager.dart';
+import 'models/unreal_dockable_tab_manager.dart';
 import 'models/unreal_property_manager.dart';
 import 'models/unreal_transaction_manager.dart';
 import 'routes.dart';
@@ -181,6 +182,10 @@ class _EpicStageAppState extends State<EpicStageApp> with TickerProviderStateMix
           create: (context) => UnrealActorCreator(context),
           dispose: (context, value) => value.dispose(),
         ),
+        Provider<UnrealDockableTabManager>(
+          create: (context) => UnrealDockableTabManager(context),
+          dispose: (context, value) => value.dispose(),
+        )
       ],
       child: MaterialApp(
         theme: UnrealTheme.makeThemeData(),
diff --git a/Engine/Extras/VirtualProduction/EpicStageApp/lib/models/settings/floating_map_settings.dart b/Engine/Extras/VirtualProduction/EpicStageApp/lib/models/settings/floating_map_settings.dart
deleted file mode 100644
diff --git a/Engine/Extras/VirtualProduction/EpicStageApp/lib/models/settings/floating_tab_settings.dart b/Engine/Extras/VirtualProduction/EpicStageApp/lib/models/settings/floating_tab_settings.dart
new file mode 100644
--- /dev/null
+++ b/Engine/Extras/VirtualProduction/EpicStageApp/lib/models/settings/floating_tab_settings.dart
@@ -0,0 +1,39 @@
+// Copyright Epic Games, Inc. All Rights Reserved.
+
+import 'package:epic_common/preferences.dart';
+import 'package:streaming_shared_preferences/streaming_shared_preferences.dart';
+
+/// Which side of the screen a tab is on.
+enum TabCurrentSide {
+  left,
+  right,
+}
+
+/// Generic holder for user settings used to configure floating tabs.
+/// required and instance of [preference] bundle and a [prefix] string indicating which tab an instance of
+/// [FloatingTabSettings] controls.
+class FloatingTabSettings {
+  FloatingTabSettings(PreferencesBundle preferences, String prefix)
+      : tabSide = preferences.persistent.getEnum<TabCurrentSide>(
+          '$prefix.tabSide',
+          defaultValue: TabCurrentSide.right,
+          enumValues: TabCurrentSide.values,
+        ),
+        yAxis = preferences.persistent.getDouble(
+          '$prefix.yAxis',
+          defaultValue: 1.0,
+        ),
+        bIsTabDocked = preferences.persistent.getBool(
+          '$prefix.bIsTabDocked',
+          defaultValue: true,
+        );
+
+  /// Which side of the screen the tab is on.
+  Preference<TabCurrentSide> tabSide;
+
+  /// The vertical position at which the tab should sit (as a fraction of its container height).
+  Preference<double> yAxis;
+
+  /// Whether the tab is docked off-screen.
+  Preference<bool> bIsTabDocked;
+}
diff --git a/Engine/Extras/VirtualProduction/EpicStageApp/lib/models/unreal_dockable_tab_manager.dart b/Engine/Extras/VirtualProduction/EpicStageApp/lib/models/unreal_dockable_tab_manager.dart
new file mode 100644
--- /dev/null
+++ b/Engine/Extras/VirtualProduction/EpicStageApp/lib/models/unreal_dockable_tab_manager.dart
@@ -0,0 +1,416 @@
+// Copyright Epic Games, Inc. All Rights Reserved.
+
+import 'dart:math' as math;
+import 'package:epic_common/preferences.dart';
+import 'package:flutter/material.dart';
+import 'package:flutter/physics.dart';
+import 'package:flutter/scheduler.dart';
+
+import '../../models/navigator_keys.dart';
+import '../../models/settings/floating_tab_settings.dart';
+import '../widgets/elements/floating_tab_base.dart';
+
+/// Simple manager for dockable tabs.
+class UnrealDockableTabManager {
+  UnrealDockableTabManager(BuildContext context) {
+    mapTabController = UnrealFloatingTabController(mapPreferencePrefix);
+    trackpadTabController = UnrealFloatingTabController(trackpadPreferencePrefix);
+  }
+
+  /// Controller for dockable map tab.
+  late final UnrealFloatingTabController mapTabController;
+
+  /// Controller for dockable trackpad tab.
+  late final UnrealFloatingTabController trackpadTabController;
+
+  /// Get relative tab position
+  TabPosition getTabRelativePosition(Offset draggedTabPosition, Offset stationaryTabPosition) {
+    return draggedTabPosition.dy > stationaryTabPosition.dy ? TabPosition.below : TabPosition.above;
+  }
+
+  void dispose() {
+    mapTabController.dispose();
+    trackpadTabController.dispose();
+  }
+}
+
+/// Controller for a Dockable tab, a string [persistentPrefix] for prefixing state stored on our persistent store is
+/// required.
+class UnrealFloatingTabController with ChangeNotifier {
+  /// prefix for controllers used when persisting state
+  final String persistentPrefix;
+
+  UnrealFloatingTabController(this.persistentPrefix);
+
+  /// Instance of tab settings.
+  FloatingTabSettings get _settings => FloatingTabSettings(
+        PreferencesBundle.of(rootNavigatorKey.currentContext!),
+        persistentPrefix,
+      );
+
+  /// Physics settings for the spring simulation when the user is dragging the tab.
+  final draggedSpringDescription = SpringDescription(mass: 40, stiffness: 1, damping: 1.2);
+
+  /// Physics settings for the spring simulation when the user releases the tab.
+  final releasedSpringDescription = SpringDescription(mass: 40, stiffness: 1, damping: 1);
+
+  /// Physics settings for the spring simulation when the tab is being docked.
+  final dockedSpringDescription = SpringDescription(mass: 60, stiffness: 1, damping: 1);
+
+  /// When this portion of the tab's width is offscreen, dock it.
+  static const double _dockThreshold = 1 / 3;
+
+  /// Which side the tab was last docked on.
+  DockSide dockSide = DockSide.left;
+
+  /// State of the tab window's docking.
+  DockState dockState = DockState.undocked;
+
+  /// The position that the tab's top-left corner should move towards with spring motion.
+  Offset targetTabPosition = Offset.zero;
+
+  /// The current position of the tab's top-left corner.
+  Offset tabPosition = Offset.zero;
+
+  /// The current velocity of the tab.
+  Offset tabVelocity = Offset.zero;
+
+  /// The user's last drag position in global space.
+  Offset lastDragPosition = Offset.zero;
+
+  /// The last known size of the containing render box.
+  Size containerSize = Size.zero;
+
+  /// The elapsed ticker time in seconds when the springs were last updated.
+  double lastSpringTime = 0;
+
+  /// The last value of [lastSpringTime] when the springs' target positions were moved.
+  double lastSpringMoveTime = 0;
+
+  /// Spring controlling the X axis.
+  SpringSimulation? xSpring;
+
+  /// Spring controlling the Y axis.
+  SpringSimulation? ySpring;
+
+  /// Ticker used to update springs.
+  Ticker? ticker;
+
+  /// The size of the tab to display given the containing widget's size.
+  Size get tabSize {
+    final Size screenSize = MediaQuery.of(rootNavigatorKey.currentContext!).size;
+    final double previewHeight = math.min(screenSize.width, screenSize.height) * tabSizeScale;
+
+    return Size(previewHeight * tabAspectRatio, previewHeight);
+  }
+
+  /// True if the tab is currently docked or in the process of docking.
+  bool get bIsDockedOrDocking => dockState == DockState.docking || dockState == DockState.docked;
+
+  /// True if the tab is currently undocked or in the process of undocking.
+  bool get bIsUndockedOrUndocking => dockState == DockState.undocking || dockState == DockState.undocked;
+
+  /// True if the tab is not docked and offscreen.
+  bool get bIsTabVisible {
+    if (dockState != DockState.docked) {
+      return true;
+    }
+
+    // Check if the tab is actually within visible bounds (e.g. still moving to docked position or being dragged out
+    // of it).
+    return ((tabPosition.dx + tabSize.width - dockTabSize.width) > 0 &&
+        (tabPosition.dx + dockTabSize.width) < containerSize.width);
+  }
+
+  /// Initialize the position of the tab once we a valid size to reference.
+  /// This should NOT call setState, as it's called directly from the first build().
+  void initPosition() {
+    final restoredYPosition = _settings.yAxis.getValue() * containerSize.height;
+    final TabCurrentSide mapSide = _settings.tabSide.getValue();
+
+    if (_settings.bIsTabDocked.getValue()) {
+      switch (mapSide) {
+        case TabCurrentSide.left:
+          dockSide = DockSide.left;
+          break;
+
+        case TabCurrentSide.right:
+          dockSide = DockSide.right;
+          break;
+      }
+
+      dockState = DockState.docked;
+      targetTabPosition = tabPosition = getDockedTabPosition(dockY: restoredYPosition);
+    } else {
+      // Make an approximate position for the tab, which we'll then snap to the nearest corner
+      late final double tempXPosition;
+      switch (mapSide) {
+        case TabCurrentSide.left:
+          tempXPosition = 0;
+          break;
+
+        case TabCurrentSide.right:
+          tempXPosition = containerSize.width - tabSize.width;
+          break;
+      }
+
+      final tempPosition = Offset(tempXPosition, restoredYPosition);
+
+      dockState = DockState.undocked;
+      targetTabPosition = tabPosition = getTabNearestCornerPosition(tempPosition);
+    }
+  }
+
+  /// Called when the containing widget is resized and we need to readjust accordingly.
+  void onContainerResized(Size oldContainerSize) {
+    if (containerSize.height == 0 || containerSize.width == 0) {
+      return;
+    }
+
+    final Offset scaledPosition = Offset(
+      targetTabPosition.dx * (containerSize.width / oldContainerSize.width),
+      targetTabPosition.dy * (containerSize.height / oldContainerSize.height),
+    );
+
+    if (bIsDockedOrDocking) {
+      // Stay docked and snap to the new side
+      targetTabPosition = tabPosition = getDockedTabPosition(dockY: scaledPosition.dy);
+      notifyListeners();
+    } else {
+      // Move smoothly to the new position
+      moveTabToNearestCorner(scaledPosition);
+    }
+  }
+
+  /// Update spring tab position using springs.
+  void onTick(Duration elapsed) {
+    if (xSpring == null || ySpring == null) {
+      return;
+    }
+
+    lastSpringTime = elapsed.inMilliseconds / 1000.0;
+
+    final double springTime = lastSpringTime - lastSpringMoveTime;
+
+    tabVelocity = Offset(xSpring!.dx(springTime), ySpring!.dx(springTime));
+    tabPosition = Offset(xSpring!.x(springTime), ySpring!.x(springTime));
+    notifyListeners();
+
+    // If both springs are at rest, stop ticking to save performance
+    if (xSpring!.isDone(springTime) && ySpring!.isDone(springTime)) {
+      ticker?.stop();
+      lastSpringMoveTime = 0;
+      lastSpringTime = 0;
+    }
+  }
+
+  /// Called when the animation to reveal/hide a dock tab completes.
+  void onDockTabAnimationEnd() {
+    switch (dockState) {
+      case DockState.undocking:
+        dockState = DockState.undocked;
+        notifyListeners();
+        break;
+
+      case DockState.docking:
+        dockState = DockState.docked;
+        notifyListeners();
+        break;
+
+      default:
+        break;
+    }
+  }
+
+  /// Move the target position for the tab to a new location and adjust springs accordingly.
+  void moveTab(Offset newTarget, {required SpringDescription springDescription}) {
+    // Clamp both axes in a valid range
+    targetTabPosition = Offset(
+      newTarget.dx.clamp(-tabSize.width + dockTabSize.width, containerSize.width - dockTabSize.width),
+      getClampedTabYPosition(newTarget.dy),
+    );
+
+    xSpring = SpringSimulation(
+      springDescription,
+      tabPosition.dx,
+      targetTabPosition.dx,
+      tabVelocity.dx,
+    );
+
+    ySpring = SpringSimulation(
+      springDescription,
+      tabPosition.dy,
+      targetTabPosition.dy,
+      tabVelocity.dy,
+    );
+
+    lastSpringMoveTime = lastSpringTime;
+
+    if (ticker?.isActive == false) {
+      ticker?.start();
+    }
+
+    _updateTabSettingsYAxis(newTarget.dy);
+  }
+
+  /// Move the tab target to the corner nearest to a local [position].
+  void moveTabToNearestCorner(Offset position) {
+    moveTab(
+      getTabNearestCornerPosition(position),
+      springDescription: releasedSpringDescription,
+    );
+
+    _settings.tabSide.setValue(
+      (targetTabPosition.dx < containerSize.width / 2) ? TabCurrentSide.left : TabCurrentSide.right,
+    );
+
+    undockTab();
+  }
+
+  /// Move the tab to the corner farthest to a local [position].
+  void moveTabToFarthestCorner(Offset position) {
+    dockTab(side, dockY: getFarthestCornerPosition(position).dy);
+    _settings.tabSide
+        .setValue((targetTabPosition.dx < containerSize.width / 2) ? TabCurrentSide.left : TabCurrentSide.right);
+  }
+
+  /// Reveal the dock tab and move to the docked state to the given [side].
+  /// If [moveOffscreen] is true, move the tab off-screen, optionally using the Y position specified by [dockY].
+  void dockTab(DockSide side, {bool moveOffscreen = true, double? dockY}) {
+    dockY = dockY ?? tabPosition.dy;
+
+    if (bIsUndockedOrUndocking) {
+      // Trigger the docking animation
+      dockSide = side;
+      dockState = DockState.docking;
+
+      notifyListeners();
+
+      // Update user settings
+      switch (side) {
+        case DockSide.left:
+          _settings.tabSide.setValue(TabCurrentSide.left);
+          break;
+
+        case DockSide.right:
+          _settings.tabSide.setValue(TabCurrentSide.right);
+          break;
+
+        default:
+          break;
+      }
+
+      _settings.bIsTabDocked.setValue(true);
+      _updateTabSettingsYAxis(dockY);
+    }
+
+    if (!moveOffscreen) {
+      return;
+    }
+
+    moveTab(
+      getDockedTabPosition(dockY: dockY),
+      springDescription: dockedSpringDescription,
+    );
+  }
+
+  /// Hide the docked tab and move to the undocking state.
+  void undockTab() {
+    if (bIsDockedOrDocking) {
+      dockState = DockState.undocking;
+      notifyListeners();
+
+      _settings.bIsTabDocked.setValue(false);
+    }
+  }
+
+  /// Check whether the tab should be docked if it's at the given X position.
+  /// Returns the side to dock to, or null if it shouldn't dock.
+  DockSide? checkDockSideAtPosition(Offset position) {
+    if (position.dx + ((1 - _dockThreshold) * tabSize.width) > containerSize.width) {
+      return DockSide.right;
+    }
+
+    if (position.dx + (_dockThreshold * tabSize.width) < 0) {
+      return DockSide.left;
+    }
+
+    return null;
+  }
+
+  /// Update the saved Y position of the docked tab.
+  void _updateTabSettingsYAxis(double dockY) {
+    if (containerSize.height != 0) {
+      _settings.yAxis.setValue(dockY / containerSize.height);
+    }
+  }
+
+  /// Clamp a Y position for the tab to a valid range.
+  double getClampedTabYPosition(double y) {
+    final double minY = edgeInset;
+    final double maxY = containerSize.height - tabSize.height - edgeInset;
+
+    return math.max(math.min(y, maxY), minY);
+  }
+
+  /// Get the position of the tab when moved to the nearest corner to a local [position].
+  Offset getTabNearestCornerPosition(Offset position) {
+    late final double targetX;
+    if (position.dx < containerSize.width / 2) {
+      targetX = edgeInset;
+    } else {
+      targetX = containerSize.width - tabSize.width - edgeInset;
+    }
+
+    late final double targetY;
+    if (position.dy < containerSize.height / 2) {
+      targetY = edgeInset;
+    } else {
+      targetY = containerSize.height - tabSize.height - edgeInset;
+    }
+
+    return Offset(targetX, targetY);
+  }
+
+  /// Get the position of the tab when moved to the farthest corner to a local [position].
+  Offset getFarthestCornerPosition(Offset position) {
+    late final double targetX;
+    if (position.dx < containerSize.width / 2) {
+      targetX = edgeInset;
+    } else {
+      targetX = containerSize.width + tabSize.width - edgeInset;
+    }
+
+    late final double targetY;
+    if (position.dy < containerSize.height / 2) {
+      targetY = containerSize.height + tabSize.height - edgeInset;
+    } else {
+      targetY = edgeInset;
+    }
+
+    return Offset(targetX, targetY);
+  }
+
+  /// Get the docked position of the tab when its Y position is [dockY].
+  Offset getDockedTabPosition({required double dockY}) {
+    // Extra amount to move just to make sure it's fully off-screen so we can stop rendering the tab
+    const double offsetFudge = 0.5;
+
+    // Move the tab to the appropriate side
+    late final double targetX;
+    switch (dockSide) {
+      case DockSide.left:
+        targetX = -tabSize.width + dockTabSize.width - offsetFudge;
+        break;
+
+      case DockSide.right:
+        targetX = containerSize.width - dockTabSize.width + offsetFudge;
+        break;
+    }
+
+    return Offset(targetX, getClampedTabYPosition(dockY));
+  }
+
+  ///Get current side of the screen where tab is currently seated, whether docked or undocked;
+  DockSide get side => tabPosition.dx > containerSize.width / 2 ? DockSide.right : DockSide.left;
+}
diff --git a/Engine/Extras/VirtualProduction/EpicStageApp/lib/models/unreal_property_controller.dart b/Engine/Extras/VirtualProduction/EpicStageApp/lib/models/unreal_property_controller.dart
new file mode 100644
--- /dev/null
+++ b/Engine/Extras/VirtualProduction/EpicStageApp/lib/models/unreal_property_controller.dart
@@ -0,0 +1,433 @@
+// Copyright Epic Games, Inc. All Rights Reserved.
+
+import 'dart:collection';
+
+import 'package:flutter/foundation.dart';
+import 'package:flutter/material.dart';
+import 'package:logging/logging.dart';
+import 'package:provider/provider.dart';
+import 'package:flutter_gen/gen_l10n/app_localizations.dart';
+
+import '../widgets/elements/delta_widget_base.dart';
+import '../widgets/elements/unreal_widget_base.dart';
+import 'property_modify_operations.dart';
+import 'unreal_property_manager.dart';
+import 'unreal_types.dart';
+import 'navigator_keys.dart';
+
+final _log = Logger('UnrealPropertyController');
+
+/// A Controller that allows modifying a list of unreal properties sharing the same name and type of all currently
+/// selected actors.
+class UnrealPropertyController<PropertyType> with ChangeNotifier {
+  /// Construct a property manager and register it with the connection manager.
+  UnrealPropertyController(BuildContext context, {bool bShouldInitTransaction = true})
+      : _propertyManager = Provider.of<UnrealPropertyManager>(context, listen: false),
+        _bShouldInitTransaction = bShouldInitTransaction {}
+
+  /// The property manager this widget uses to update its properties.
+  final UnrealPropertyManager _propertyManager;
+
+  /// Whether or not to initiate transactions internally or to allow transactions be initiated externally by calling
+  /// [beginTransaction] and [endTransaction] respectively to start & end a transaction.
+  final bool _bShouldInitTransaction;
+
+  /// List of properties currently controlled by the widget.
+  /// We can have multiple properties due to multi-selection.
+  /// properties are of the same name and type.
+  List<WidgetControlledUnrealProperty<PropertyType>?> _properties = [];
+
+  /// Unmodifiable getter getter for [_properties].
+  UnmodifiableListView<WidgetControlledUnrealProperty<PropertyType>?> get properties =>
+      UnmodifiableListView(_properties);
+
+  /// A list of properties for enabling widget properties current controlled by the widget.
+  final List<WidgetControlledUnrealProperty<bool>?> _enableProperties = [];
+
+  /// The last list of properties that were passed to the widget. This is stored so we can compare the list on rebuild
+  /// and register/unregister properties when the list changes.
+  final List<UnrealProperty> _lastUnrealProperties = [];
+
+  /// The last list of enable properties that were passed to the widget. This is stored so we can compare the list on
+  /// rebuild and register/unregister properties when the list changes.
+  final List<UnrealProperty> _lastEnableProperties = [];
+
+  /// Property's metadata as received from the engine.
+  UnrealPropertyMetadata? _propertyMetadata;
+
+  /// This is incremented whenever the set of tracked properties changes so we can ignore any
+  /// callbacks from previous versions. For example, if the user rapidly changes the list multiple times, we may get
+  /// changed value or newly tracked value callbacks for properties that we no longer care about.
+  int _latestTrackedPropertyVersion = 0;
+
+  /// Whether the properties' values are ready to be display to the user.
+  bool _bIsReady = false;
+
+  /// Whether the widget is actively in a transaction or not.
+  bool _bIsInTransaction = false;
+
+  /// The label of the property to show to users.
+  String get propertyLabel => (_propertyMetadata?.displayName ?? '');
+
+  /// The property's minimum value as received from the engine.
+  PropertyType? get engineMin => _propertyMetadata?.minValue;
+
+  /// The property's maximum value as received from the engine.
+  PropertyType? get engineMax => _propertyMetadata?.maxValue;
+
+  /// The overridden minimum value the controlled properties can reach. If null, [engineMin] will be used.
+  PropertyType? get overrideMin => null;
+
+  /// The overridden maximum value the controlled properties can reach. If null, [engineMax] will be used.
+  PropertyType? get overrideMax => null;
+
+  /// The property's possible enum values as received from the engine.
+  List<String> get propertyEnumValues => _propertyMetadata?.enumValues ?? [];
+
+  /// The type of modify operation this widget uses to apply deltas.
+  PropertyModifyOperation get modifyOperation => const AddOperation();
+
+  /// Extra meta-data properties to pass to conversion functions operating on this widget's value.
+  Map<String, dynamic>? get conversionMetadata => null;
+
+  /// Call this to indicate that the user has started interacting with the widget.
+  /// This will be automatically called for you as soon as any properties change, but you can safely call it earlier.
+  /// A [description] will be automatically generated if null is passed.
+  /// Returns true if a new transaction was started.
+  bool beginTransaction([String? description]) {
+    if (_bIsInTransaction) {
+      return false;
+    }
+
+    if (_propertyManager.beginTransaction(
+        description ?? AppLocalizations.of(rootNavigatorKey.currentContext!)!.transactionEditProperty(propertyLabel))) {
+      _bIsInTransaction = true;
+      return true;
+    }
+    return false;
+  }
+
+  /// Call this to indicate that the user has finished interacting with the widget.
+  void endTransaction() {
+    if (!_bIsInTransaction) {
+      return;
+    }
+
+    _bIsInTransaction = false;
+    _propertyManager.endTransaction();
+  }
+
+  /// Generate a list of [DeltaWidgetValueData] containing the data needed to display each control in the widget.
+  /// Values may be null if the value hasn't been received from the engine yet.
+  List<DeltaWidgetValueData<PropertyType>?> makeValueDataList() =>
+      _bIsReady ? _properties.map((property) => property?.toValueData()).toList() : [];
+
+  /// Called when the user changes controlled values.
+  void handleOnChangedByUser(List<PropertyType> deltaValues) {
+    modifyProperties(modifyOperation, values: deltaValues);
+  }
+
+  /// Called when the user presses the reset button.
+  void handleOnResetByUser(List<UnrealProperty>? enableProperties) {
+    if (!_propertyManager.beginTransaction(
+        AppLocalizations.of(rootNavigatorKey.currentContext!)!.transactionResetProperty(propertyLabel))) {
+      return;
+    }
+
+    for (final WidgetControlledUnrealProperty<PropertyType>? property in _properties) {
+      if (property != null) {
+        _propertyManager.modifyTrackedPropertyValue(property.trackedId, const ResetOperation());
+      }
+    }
+
+    if (enableProperties != null) {
+      for (final WidgetControlledUnrealProperty<bool>? property in _enableProperties) {
+        if (property != null) {
+          _propertyManager.modifyTrackedPropertyValue(property.trackedId, const ResetOperation());
+        }
+      }
+    }
+
+    _propertyManager.endTransaction();
+  }
+
+  /// Return a SingleSharedValue indicating whether there is a single value shared by all controlled properties.
+  /// If there are no controlled values, the SingleSharedValue will instead contain [defaultValue].
+  SingleSharedValue<PropertyType> getSingleSharedValue([PropertyType? defaultValue]) {
+    // Get the value of all non-null properties
+    final List<PropertyType?> controlledValues = _properties.map((property) => property?.value).toList();
+    controlledValues.removeWhere((element) => element == null);
+
+    bool bHasMultipleValues = false;
+    PropertyType? sharedValue;
+
+    if (controlledValues.length > 1) {
+      sharedValue = controlledValues[0];
+
+      // We're controlling multiple properties. Check that they all have the same value
+      for (int valueIndex = 1; valueIndex < controlledValues.length; ++valueIndex) {
+        if (controlledValues[valueIndex] != sharedValue) {
+          bHasMultipleValues = true;
+          sharedValue = null;
+          break;
+        }
+      }
+    } else if (controlledValues.isNotEmpty) {
+      // There's exactly one value, so use it directly
+      sharedValue = controlledValues[0];
+    } else {
+      // There are no properties, so return the default
+      sharedValue = defaultValue;
+    }
+
+    return SingleSharedValue(value: sharedValue, bHasMultipleValues: bHasMultipleValues);
+  }
+
+  /// Apply an [operation] to each property this controls by the amounts specified in [values].
+  /// If [values] is null, this is treated as if it was a list of null values matching the number of properties.
+  /// If [bIgnoreLimits] is true, the min and max for the widget will be ignored.
+  /// If [_bShouldInitTransaction] is false, we ought to call [beginTransaction] before calling
+  void modifyProperties(
+    PropertyModifyOperation operation, {
+    List<dynamic>? values,
+    bool bIgnoreLimits = false,
+    PropertyMinMaxBehaviour minMaxBehaviour = PropertyMinMaxBehaviour.clamp,
+    void Function()? onChangedByUser,
+  }) {
+    assert(values == null || values.length == _properties.length);
+
+    for (int index = 0; index < _properties.length; index++) {
+      modifyProperty(
+        operation,
+        index,
+        value: values?[index],
+        bIgnoreLimits: bIgnoreLimits,
+        minMaxBehaviour: minMaxBehaviour,
+      );
+    }
+
+    if (onChangedByUser != null) {
+      onChangedByUser();
+    }
+  }
+
+  /// Apply an [operation] to a single property specified by the property [index] this controls by the amounts specified
+  /// in [value].
+  /// If [value] is null, this is treated as null and applied to matching property [index].
+  /// If [bIgnoreLimits] is true, the min and max for the widget will be ignored.
+  /// If [_bShouldInitTransaction] is false, we ought to call [beginTransaction] before calling
+  void modifyProperty(
+    PropertyModifyOperation operation,
+    int index, {
+    dynamic value,
+    bool bIgnoreLimits = false,
+    PropertyMinMaxBehaviour minMaxBehaviour = PropertyMinMaxBehaviour.clamp,
+  }) {
+    if (_bShouldInitTransaction == true) {
+      if (!_bIsInTransaction) {
+        if (!beginTransaction()) {
+          // If we can't start a transaction, ignore user input or we'll get out of sync with the engine
+          return;
+        }
+      }
+    }
+
+    try {
+      //Enable the property in the engine
+      if (index < _enableProperties.length) {
+        final WidgetControlledUnrealProperty<bool>? enableProperty = _enableProperties[index];
+
+        if (enableProperty != null) {
+          if (!_propertyManager.getTrackedPropertyValue(enableProperty.trackedId)) {
+            _propertyManager.modifyTrackedPropertyValue(enableProperty.trackedId, const SetOperation(),
+                deltaValue: true);
+          }
+        }
+      }
+
+      final WidgetControlledUnrealProperty<PropertyType>? property = _properties[index];
+
+      _propertyManager.modifyTrackedPropertyValue(
+        property!.trackedId,
+        operation,
+        deltaValue: value,
+        minMaxBehaviour: bIgnoreLimits ? PropertyMinMaxBehaviour.ignore : minMaxBehaviour,
+        overrideMin: overrideMin,
+        overrideMax: overrideMax,
+      );
+    } on Exception catch (e) {
+      _log.severe('failed to modify tracked property', e);
+    }
+  }
+
+  /// Check whether the widget's list of properties has changed, and if so, update our list of tracked properties.
+  void trackAllPropertiesIfChanged(List<UnrealProperty> unrealProperties, List<UnrealProperty>? enableProperties) {
+    bool bHavePropertiesChanged = unrealProperties.length != _lastUnrealProperties.length ||
+        (enableProperties?.length ?? 0) != _lastEnableProperties.length;
+
+    if (!bHavePropertiesChanged) {
+      bHavePropertiesChanged = !listEquals(unrealProperties, _lastUnrealProperties);
+    }
+
+    if (!bHavePropertiesChanged && enableProperties != null) {
+      bHavePropertiesChanged = !listEquals(enableProperties, _lastEnableProperties);
+    }
+
+    if (!bHavePropertiesChanged) {
+      return;
+    }
+
+    trackAllProperties(unrealProperties, enableProperties);
+  }
+
+  /// Resets the internal list of properties and starts tracking them again.
+  void trackAllProperties(List<UnrealProperty> unrealProperties, [List<UnrealProperty>? enableProperties]) {
+    _bIsReady = false;
+    ++_latestTrackedPropertyVersion;
+
+    // The tracked property version when we started this function, so we can ignore callbacks later in the function
+    // if they come back too late to be relevant.
+    final int trackedPropertyVersion = _latestTrackedPropertyVersion;
+
+    // Futures that will complete when each property is exposed.
+    final List<Future<void>> exposeFutures = [];
+
+    trackRelevantProperties<PropertyType>(
+      trackedProperties: _properties,
+      oldProperties: _lastUnrealProperties,
+      newProperties: unrealProperties,
+      exposeFutures: exposeFutures,
+      callback: _handleOnManagedValueChanged,
+    );
+
+    trackRelevantProperties<bool>(
+      trackedProperties: _enableProperties,
+      oldProperties: _lastEnableProperties,
+      newProperties: enableProperties ?? [],
+      exposeFutures: exposeFutures,
+      callback: _handleOnEnableValueChanged,
+    );
+
+    // Wait for all properties to be subscribed, then mark the widget as ready to be used
+    Future.wait(exposeFutures).then((_) {
+      if (_latestTrackedPropertyVersion != trackedPropertyVersion) {
+        return;
+      }
+
+      _onPropertiesReady();
+      notifyListeners();
+    }).onError((error, stackTrace) {
+      _log.severe('Error while waiting for Unreal widget properties', error, stackTrace);
+    });
+  }
+
+  /// Given a list of [trackedProperties], a corresponding list of [oldProperties] that are being tracked, and a list
+  /// of [newProperties] that are about to be tracked, stop tracking any properties that are no longer relevant and
+  /// start any properties that are newly relevant. [exposeFutures] will be filled with futures for each newly-tracked
+  /// property that will complete when the property is exposed, and [callback] will be called in the future whenever a
+  /// property's value changes.
+  void trackRelevantProperties<T>({
+    required List<WidgetControlledUnrealProperty<T>?> trackedProperties,
+    required List<UnrealProperty> oldProperties,
+    required List<UnrealProperty> newProperties,
+    required List<Future<void>> exposeFutures,
+    required Function(int trackedPropertyVersion, int index, dynamic newValue) callback,
+  }) {
+    assert(oldProperties.length == trackedProperties.length);
+
+    // Stop tracking all properties. We either aren't tracking them anymore, or we may need to update their
+    // propertyIndex if the number tracked has changed size, in which case we need to recreate the callback anyway.
+    for (int propertyIndex = 0; propertyIndex < oldProperties.length; ++propertyIndex) {
+      final WidgetControlledUnrealProperty<T>? property = trackedProperties[propertyIndex];
+
+      if (property != null) {
+        _propertyManager.stopTrackingProperty(property.trackedId, property.changedCallback);
+      }
+    }
+
+    oldProperties.clear();
+    oldProperties.addAll(newProperties);
+
+    trackedProperties.length = newProperties.length;
+
+    // Track each property we want to control with the remote property manager.
+    for (int propertyIndex = 0; propertyIndex < newProperties.length; ++propertyIndex) {
+      final UnrealProperty property = newProperties[propertyIndex];
+
+      // Create a callback function for this property
+      changedCallback(newValue) => callback(_latestTrackedPropertyVersion, propertyIndex, newValue);
+
+      // Start tracking the property in the property manager
+      final TrackedPropertyId trackedId =
+          _propertyManager.trackProperty(property, changedCallback, conversionMetadata: conversionMetadata);
+
+      final newProperty = WidgetControlledUnrealProperty<T>(
+        trackedId: trackedId,
+        propertyIndex: propertyIndex,
+        changedCallback: changedCallback,
+        value: _propertyManager.getTrackedPropertyValue(trackedId),
+        ownerName: property.objectName,
+      );
+      trackedProperties[propertyIndex] = newProperty;
+
+      // Add to the list of properties we're waiting for
+      final Future<void> exposeFuture = _propertyManager.waitForProperty(trackedId);
+      exposeFutures.add(exposeFuture);
+
+      exposeFuture.then((_) {
+        newProperty.value = _propertyManager.getTrackedPropertyValue(trackedId);
+      }).onError((error, stackTrace) {
+        _log.severe(
+            'Error while exposing property ${property.propertyName} on ${property.objectName}', error, stackTrace);
+      });
+    }
+  }
+
+  /// Called when all properties have been subscribed and the widget is ready to be interacted with.
+  void _onPropertiesReady() {
+    if (_properties.isNotEmpty) {
+      // Properties could differ in their metadata, but for now we'll assume they share the same values
+      final WidgetControlledUnrealProperty<PropertyType>? firstProperty = _properties[0];
+
+      if (firstProperty == null) {
+        assert(false, 'If an Unreal widget is ready, all properties should be non-null');
+        return;
+      }
+
+      _propertyMetadata = _propertyManager.getTrackedPropertyMetaData(firstProperty.trackedId);
+    }
+    _bIsReady = true;
+    notifyListeners();
+  }
+
+  /// Called when the property manager reports a [newValue] for the property with the given [index] in the widget.
+  /// If [trackedPropertyVersion] does not match [_latestTrackedPropertyVersion], the change will be ignored.
+  void _handleOnManagedValueChanged(int trackedPropertyVersion, int index, dynamic newValue) {
+    if (trackedPropertyVersion != _latestTrackedPropertyVersion) {
+      return;
+    }
+
+    _properties[index]?.value = newValue;
+    notifyListeners();
+  }
+
+  /// Called when an enable value for one of the properties changes.
+  void _handleOnEnableValueChanged(int trackedPropertyVersion, int index, dynamic newValue) {}
+
+  @override
+  void dispose() {
+    for (final WidgetControlledUnrealProperty<PropertyType>? property in properties) {
+      if (property != null) {
+        _propertyManager.stopTrackingProperty(property.trackedId, property.changedCallback);
+      }
+    }
+
+    for (final WidgetControlledUnrealProperty<bool>? property in _enableProperties) {
+      if (property != null) {
+        _propertyManager.stopTrackingProperty(property.trackedId, property.changedCallback);
+      }
+    }
+
+    super.dispose();
+  }
+}
diff --git a/Engine/Extras/VirtualProduction/EpicStageApp/lib/widgets/elements/floating_map_preview.dart b/Engine/Extras/VirtualProduction/EpicStageApp/lib/widgets/elements/floating_map_preview.dart
--- a/Engine/Extras/VirtualProduction/EpicStageApp/lib/widgets/elements/floating_map_preview.dart
+++ b/Engine/Extras/VirtualProduction/EpicStageApp/lib/widgets/elements/floating_map_preview.dart
@@ -1,688 +1,55 @@
 // Copyright Epic Games, Inc. All Rights Reserved.
 
-import 'dart:math' as math;
-
-import 'package:epic_common/preferences.dart';
 import 'package:epic_common/widgets.dart';
 import 'package:flutter/material.dart';
-import 'package:flutter/physics.dart';
-import 'package:flutter/scheduler.dart';
+import 'package:provider/provider.dart';
 
-import '../../models/settings/floating_map_settings.dart';
+import 'floating_tab_base.dart';
 import 'lightcard_map.dart';
 
-/// Sides of the screen on which the preview can be docked.
-enum _DockSide {
-  left,
-  right,
-}
-
-/// Possible states for the floating preview in terms of docking.
-enum _DockState {
-  /// The preview is on-screen and no dock tabs are visible.
-  undocked,
-
-  /// The preview has entered the docking range and is revealing its tab.
-  docking,
-
-  /// The preview has revealed its docking tab. It may not actually be off-screen yet.
-  docked,
-
-  /// The preview has exited the docking range and is hiding its tab.
-  undocking,
-}
-
-/// The smallest dimension of the container will be multiplied by this value to determine the height of the preview.
-const double _previewSizeScale = 0.25;
-
-/// The height of the preview will be multiplied by this value to determine its width.
-const double _previewAspectRatio = 16.0 / 9.0;
-
-/// How far to inset the preview from the edge of the screen.
-const double _edgeInset = 8.0;
-
-/// The size of the tab shown when the preview is docked.
-const _dockTabSize = Size(32, 80);
-
 /// A floating preview of the stage map that the user can drag to any corner of the screen and dock to the side.
 class FloatingMapPreview extends StatefulWidget {
   const FloatingMapPreview({Key? key}) : super(key: key);
 
+  static GlobalKey map = GlobalKey(debugLabel: 'floating_map');
+
   @override
   State<FloatingMapPreview> createState() => _FloatingMapPreviewState();
 }
 
-class _FloatingMapPreviewState extends State<FloatingMapPreview> with SingleTickerProviderStateMixin {
-  /// Physics settings for the spring simulation when the user is dragging the preview.
-  static const _draggedSpringDescription = SpringDescription(mass: 40, stiffness: 1, damping: 1.2);
-
-  /// Physics settings for the spring simulation when the user releases the preview.
-  static const _releasedSpringDescription = SpringDescription(mass: 40, stiffness: 1, damping: 1);
-
-  /// Physics settings for the spring simulation when the preview is being docked.
-  static const _dockedSpringDescription = SpringDescription(mass: 60, stiffness: 1, damping: 1);
-
-  /// When this portion of the preview's width is offscreen, dock it.
-  static const double _dockThreshold = 1 / 3;
-
-  /// Which side the preview was last docked on.
-  _DockSide _dockSide = _DockSide.left;
-
-  /// State of the preview window's docking.
-  _DockState _dockState = _DockState.undocked;
-
-  /// The position that the preview's top-left corner should move towards with spring motion.
-  Offset _targetPreviewPosition = Offset.zero;
-
-  /// The current position of the preview's top-left corner.
-  Offset _previewPosition = Offset.zero;
-
-  /// The current velocity of the preview.
-  Offset _previewVelocity = Offset.zero;
-
-  /// The user's last drag position in global space.
-  Offset _lastDragPosition = Offset.zero;
-
-  /// The last known size of the containing render box.
-  Size _containerSize = Size.zero;
-
-  /// The elapsed ticker time in seconds when the springs were last updated.
-  double _lastSpringTime = 0;
-
-  /// The last value of [_lastSpringTime] when the springs' target positions were moved.
-  double _lastSpringMoveTime = 0;
-
-  /// Spring controlling the X axis.
-  SpringSimulation? _xSpring;
-
-  /// Spring controlling the Y axis.
-  SpringSimulation? _ySpring;
-
-  /// Ticker used to update springs.
-  late final Ticker _ticker;
-
-  /// Settings for this preview map.
-  late final FloatingMapSettings _settings;
-
-  /// The size of the preview to display given the containing widget's size.
-  Size get _previewSize {
-    final Size screenSize = MediaQuery.of(context).size;
-    final double previewHeight = math.min(screenSize.width, screenSize.height) * _previewSizeScale;
-
-    return Size(previewHeight * _previewAspectRatio, previewHeight);
-  }
-
-  /// True if the preview is currently docked or in the process of docking.
-  bool get _bIsDockedOrDocking => _dockState == _DockState.docking || _dockState == _DockState.docked;
-
-  /// True if the preview is currently undocked or in the process of undocking.
-  bool get _bIsUndockedOrUndocking => _dockState == _DockState.undocking || _dockState == _DockState.undocked;
-
-  /// True if the preview is not docked and offscreen.
-  bool get _bIsPreviewVisible {
-    if (_dockState != _DockState.docked) {
-      return true;
-    }
-
-    // Check if the preview is actually within visible bounds (e.g. still moving to docked position or being dragged out
-    // of it).
-    return ((_previewPosition.dx + _previewSize.width - _dockTabSize.width) > 0 &&
-        (_previewPosition.dx + _dockTabSize.width) < _containerSize.width);
-  }
-
-  @override
-  Widget build(BuildContext context) {
-    return LayoutBuilder(builder: (BuildContext context, BoxConstraints constraints) {
-      final Size newSize = Size(constraints.maxWidth, constraints.maxHeight);
-
-      // Handle first size/resize
-      if (newSize != _containerSize) {
-        final Size oldSize = _containerSize;
-        _containerSize = newSize;
-
-        if (oldSize == Size.zero) {
-          _initPosition();
-        } else {
-          WidgetsBinding.instance.addPostFrameCallback((_) => _onContainerResized(oldSize));
-        }
-      }
-
-      // Determine the offset of the main preview window, which shifts to the side to reveal dock tabs when docked.
-      double dockOffset = 0;
-      if (_bIsDockedOrDocking) {
-        switch (_dockSide) {
-          case _DockSide.left:
-            dockOffset = -_dockTabSize.width;
-            break;
-
-          case _DockSide.right:
-            dockOffset = _dockTabSize.width;
-            break;
-        }
-      }
-
-      return Stack(
-        clipBehavior: Clip.hardEdge,
-        children: [
-          // Position a box containing the preview, dock tabs, and gesture detector
-          Positioned(
-            top: _previewPosition.dy,
-            left: _previewPosition.dx,
-            child: SizedBox(
-              width: _previewSize.width,
-              height: _previewSize.height,
-              child: Stack(
-                clipBehavior: Clip.none,
-                children: [
-                  // Tab shown on the left when docked on the right
-                  if (_dockState != _DockState.undocked && _dockSide == _DockSide.right)
-                    Positioned(
-                      key: Key('LeftDockTab'),
-                      left: 0,
-                      top: (_previewSize.height - _dockTabSize.height) / 2,
-                      child: _DockTab(_DockSide.right),
-                    ),
-
-                  // Tab shown on the right when docked on the left
-                  if (_dockState != _DockState.undocked && _dockSide == _DockSide.left)
-                    Positioned(
-                      key: Key('RightDockTab'),
-                      left: _previewSize.width - _dockTabSize.width,
-                      top: (_previewSize.height - _dockTabSize.height) / 2,
-                      child: _DockTab(_DockSide.left),
-                    ),
-
-                  // Preview window, which shifts to reveal a tab when docked and hides when offscreen
-                  AnimatedPositioned(
-                    key: Key('PreviewMover'),
-                    left: dockOffset,
-                    duration: const Duration(milliseconds: 200),
-                    child: _bIsPreviewVisible
-                        ? _MapPreview(
-                            key: Key('Preview'),
-                            size: _previewSize,
-                            bShowShadow: _bIsUndockedOrUndocking,
-                          )
-                        : SizedBox(),
-                    onEnd: _onDockTabAnimationEnd,
-                  ),
-
-                  // Gesture detector, which matches the preview rectangle when undocked and the revealed tab when docked
-                  Positioned(
-                    key: Key('GestureDetector'),
-
-                    // When docked to left, offset to line up with the tab on the right side of the preview
-                    left: (_dockState == _DockState.docked && _dockSide == _DockSide.left)
-                        ? (_previewSize.width - _dockTabSize.width)
-                        : 0,
-
-                    // When docked, offset to line up with the top of the tab
-                    top: _dockState == _DockState.docked ? ((_previewSize.height - _dockTabSize.height) / 2) : 0,
-
-                    child: SizedBox(
-                      // When docked, match size to the tab
-                      width: _dockState == _DockState.docked ? _dockTabSize.width : _previewSize.width,
-                      height: _dockState == _DockState.docked ? _dockTabSize.height : _previewSize.height,
-
-                      child: GestureDetector(
-                        onTap: _onTap,
-                        onPanUpdate: _updateDrag,
-                        onPanEnd: (details) => _endDrag(details.velocity),
-                        onPanCancel: () => _endDrag(Velocity.zero),
-                      ),
-                    ),
-                  ),
-                ],
-              ),
-            ),
-          ),
-        ],
-      );
-    });
-  }
-
+class _FloatingMapPreviewState extends State<FloatingMapPreview> with SingleTickerProviderStateMixin, FloatingTabBase {
   @override
   void initState() {
     super.initState();
 
-    _ticker = createTicker(_onTick);
-    _settings = FloatingMapSettings(PreferencesBundle.of(context));
-  }
-
-  @override
-  void dispose() {
-    _ticker.dispose();
-
-    super.dispose();
-  }
-
-  /// Initialize the position of the preview once we a valid size to reference.
-  /// This should NOT call setState, as it's called directly from the first build().
-  void _initPosition() {
-    final restoredYPosition = _settings.mapY.getValue() * _containerSize.height;
-    final PreviewMapSide mapSide = _settings.mapSide.getValue();
-
-    if (_settings.bIsMapDocked.getValue()) {
-      switch (mapSide) {
-        case PreviewMapSide.left:
-          _dockSide = _DockSide.left;
-          break;
-
-        case PreviewMapSide.right:
-          _dockSide = _DockSide.right;
-          break;
-      }
-
-      _dockState = _DockState.docked;
-      _targetPreviewPosition = _previewPosition = _getDockedPreviewPosition(dockY: restoredYPosition);
-    } else {
-      // Make an approximate position for the preview, which we'll then snap to the nearest corner
-      late final double tempXPosition;
-      switch (mapSide) {
-        case PreviewMapSide.left:
-          tempXPosition = 0;
-          break;
-
-        case PreviewMapSide.right:
-          tempXPosition = _containerSize.width - _previewSize.width;
-          break;
-      }
-
-      final tempPosition = Offset(tempXPosition, restoredYPosition);
-
-      _dockState = _DockState.undocked;
-      _targetPreviewPosition = _previewPosition = _getNearestPreviewCornerPosition(tempPosition);
-    }
-  }
-
-  /// Called when the containing widget is resized and we need to readjust accordingly.
-  void _onContainerResized(Size oldContainerSize) {
-    if (_containerSize.height == 0 || _containerSize.width == 0) {
-      return;
-    }
-
-    final Offset scaledPosition = Offset(
-      _targetPreviewPosition.dx * (_containerSize.width / oldContainerSize.width),
-      _targetPreviewPosition.dy * (_containerSize.height / oldContainerSize.height),
-    );
-
-    if (_bIsDockedOrDocking) {
-      // Stay docked and snap to the new side
-      setState(() {
-        _targetPreviewPosition = _previewPosition = _getDockedPreviewPosition(dockY: scaledPosition.dy);
-      });
-    } else {
-      // Move smoothly to the new position
-      _movePreviewTargetToNearestCorner(scaledPosition);
-    }
-  }
-
-  /// Handle a tap on the preview.
-  void _onTap() {
-    if (_bIsUndockedOrUndocking) {
-      return;
-    }
-
-    // Currently docked, so undock by moving to the nearest corner.
-    _movePreviewTargetToNearestCorner(_targetPreviewPosition);
-  }
-
-  /// Update the dragged preview position.
-  void _updateDrag(DragUpdateDetails details) {
-    _lastDragPosition = details.globalPosition;
+    tabController = manager.mapTabController;
 
-    _movePreviewTarget(
-      _targetPreviewPosition + Offset(details.delta.dx, details.delta.dy),
-      springDescription: _draggedSpringDescription,
-    );
-
-    final _DockSide? newDockSide = _checkDockSideAtPosition(_targetPreviewPosition);
-    if (newDockSide != null) {
-      // If this will dock, show the tab
-      _dockPreview(newDockSide, moveOffscreen: false);
-    } else {
-      // No longer docked, so move to undocking state
-      _undockPreview();
-    }
-  }
-
-  /// Stop dragging the preview and let it come to rest.
-  void _endDrag(Velocity velocity) {
-    final RenderBox? renderBox = context.findRenderObject() as RenderBox?;
-    if (renderBox == null) {
-      _movePreviewTarget(
-        Offset(_edgeInset, _edgeInset),
-        springDescription: _releasedSpringDescription,
-      );
-      return;
-    }
-
-    // Check ahead from the preview's position so the user can "fling" the preview off the screen
-    const dockCheckAheadSeconds = 0.1;
-    final Offset dockCheckPosition = _previewPosition + velocity.pixelsPerSecond * dockCheckAheadSeconds;
-
-    // If near the edges, dock the preview
-    final _DockSide? newDockSide = _checkDockSideAtPosition(dockCheckPosition);
-    if (newDockSide != null) {
-      _dockPreview(newDockSide, dockY: dockCheckPosition.dy);
-      return;
-    }
-
-    // Otherwise, move to the nearest corner. In this case, use the user's drag position instead of the preview
-    // position since flinging feels more accurate that way.
-    const cornerCheckAheadSeconds = 0.15;
-    final Offset cornerCheckPosition = _lastDragPosition + velocity.pixelsPerSecond * cornerCheckAheadSeconds;
-    final Offset finalLocalPosition = renderBox.globalToLocal(cornerCheckPosition);
-
-    _movePreviewTargetToNearestCorner(finalLocalPosition);
-  }
-
-  /// Update spring preview position using springs.
-  void _onTick(Duration elapsed) {
-    if (_xSpring == null || _ySpring == null) {
-      return;
-    }
-
-    _lastSpringTime = elapsed.inMilliseconds / 1000.0;
-
-    final double springTime = _lastSpringTime - _lastSpringMoveTime;
-
-    setState(() {
-      _previewVelocity = Offset(_xSpring!.dx(springTime), _ySpring!.dx(springTime));
-      _previewPosition = Offset(_xSpring!.x(springTime), _ySpring!.x(springTime));
-    });
-
-    // If both springs are at rest, stop ticking to save performance
-    if (_xSpring!.isDone(springTime) && _ySpring!.isDone(springTime)) {
-      _ticker.stop();
-      _lastSpringMoveTime = 0;
-      _lastSpringTime = 0;
-    }
-  }
-
-  /// Called when the animation to reveal/hide a dock tab completes.
-  void _onDockTabAnimationEnd() {
-    switch (_dockState) {
-      case _DockState.undocking:
-        setState(() {
-          _dockState = _DockState.undocked;
-        });
-        break;
-
-      case _DockState.docking:
-        setState(() {
-          _dockState = _DockState.docked;
-        });
-        break;
-
-      default:
-        break;
-    }
-  }
-
-  /// Move the target position for the preview to a new location and adjust springs accordingly.
-  void _movePreviewTarget(Offset newTarget, {required SpringDescription springDescription}) {
-    // Clamp both axes in a valid range
-    _targetPreviewPosition = Offset(
-      newTarget.dx.clamp(-_previewSize.width + _dockTabSize.width, _containerSize.width - _dockTabSize.width),
-      _getClampedPreviewYPosition(newTarget.dy),
-    );
-
-    _xSpring = SpringSimulation(
-      springDescription,
-      _previewPosition.dx,
-      _targetPreviewPosition.dx,
-      _previewVelocity.dx,
-    );
-
-    _ySpring = SpringSimulation(
-      springDescription,
-      _previewPosition.dy,
-      _targetPreviewPosition.dy,
-      _previewVelocity.dy,
-    );
-
-    _lastSpringMoveTime = _lastSpringTime;
-
-    if (!_ticker.isActive) {
-      _ticker.start();
-    }
-
-    _updateSettingsPreviewMapY(newTarget.dy);
-  }
-
-  /// Move the preview target to the corner nearest to a local [position].
-  void _movePreviewTargetToNearestCorner(Offset position) {
-    _movePreviewTarget(
-      _getNearestPreviewCornerPosition(position),
-      springDescription: _releasedSpringDescription,
-    );
-
-    _settings.mapSide.setValue(
-      (_targetPreviewPosition.dx < _containerSize.width / 2) ? PreviewMapSide.left : PreviewMapSide.right,
-    );
-
-    _undockPreview();
-  }
-
-  /// Reveal the dock tab and move to the docked state to the given [side].
-  /// If [moveOffscreen] is true, move the preview off-screen, optionally using the Y position specified by [dockY].
-  void _dockPreview(_DockSide side, {bool moveOffscreen = true, double? dockY}) {
-    dockY = dockY ?? _previewPosition.dy;
-
-    if (_bIsUndockedOrUndocking) {
-      // Trigger the docking animation
-      setState(() {
-        _dockSide = side;
-        _dockState = _DockState.docking;
-      });
-
-      // Update user settings
-      switch (side) {
-        case _DockSide.left:
-          _settings.mapSide.setValue(PreviewMapSide.left);
-          break;
+    initStationaryController(manager.trackpadTabController);
 
-        case _DockSide.right:
-          _settings.mapSide.setValue(PreviewMapSide.right);
-          break;
+    tabController.ticker = createTicker(onTick);
 
-        default:
-          break;
-      }
-
-      _settings.bIsMapDocked.setValue(true);
-      _updateSettingsPreviewMapY(dockY);
-    }
-
-    if (!moveOffscreen) {
-      return;
-    }
-
-    _movePreviewTarget(
-      _getDockedPreviewPosition(dockY: dockY),
-      springDescription: _dockedSpringDescription,
-    );
-  }
-
-  /// Hide the docked tab and move to the undocking state.
-  void _undockPreview() {
-    if (_bIsDockedOrDocking) {
-      setState(() {
-        _dockState = _DockState.undocking;
-      });
-
-      _settings.bIsMapDocked.setValue(false);
-    }
-  }
-
-  /// Check whether the preview should be docked if it's at the given X position.
-  /// Returns the side to dock to, or null if it shouldn't dock.
-  _DockSide? _checkDockSideAtPosition(Offset position) {
-    if (position.dx + ((1 - _dockThreshold) * _previewSize.width) > _containerSize.width) {
-      return _DockSide.right;
-    }
-
-    if (position.dx + (_dockThreshold * _previewSize.width) < 0) {
-      return _DockSide.left;
-    }
-
-    return null;
-  }
-
-  /// Update the saved Y position of the docked preview.
-  void _updateSettingsPreviewMapY(double dockY) {
-    _settings.mapY.setValue(dockY / _containerSize.height);
-  }
-
-  /// Clamp a Y position for the preview to a valid range.
-  double _getClampedPreviewYPosition(double y) {
-    final double minY = _edgeInset;
-    final double maxY = _containerSize.height - _previewSize.height - _edgeInset;
-
-    return math.max(math.min(y, maxY), minY);
-  }
-
-  /// Get the position of the preview when moved to the nearest corner to a local [position].
-  Offset _getNearestPreviewCornerPosition(Offset position) {
-    late final double targetX;
-    if (position.dx < _containerSize.width / 2) {
-      targetX = _edgeInset;
-    } else {
-      targetX = _containerSize.width - _previewSize.width - _edgeInset;
-    }
-
-    late final double targetY;
-    if (position.dy < _containerSize.height / 2) {
-      targetY = _edgeInset;
-    } else {
-      targetY = _containerSize.height - _previewSize.height - _edgeInset;
-    }
-
-    return Offset(targetX, targetY);
-  }
-
-  /// Get the docked position of the preview when its Y position is [dockY].
-  Offset _getDockedPreviewPosition({required double dockY}) {
-    // Extra amount to move just to make sure it's fully off-screen so we can stop rendering the preview
-    const double offsetFudge = 0.5;
-
-    // Move the preview to the appropriate side
-    late final double targetX;
-    switch (_dockSide) {
-      case _DockSide.left:
-        targetX = -_previewSize.width + _dockTabSize.width - offsetFudge;
-        break;
-
-      case _DockSide.right:
-        targetX = _containerSize.width - _dockTabSize.width + offsetFudge;
-        break;
-    }
-
-    return Offset(targetX, _getClampedPreviewYPosition(dockY));
+    /// init dragged tab relative position.
+    tabRelativePosition =
+        manager.getTabRelativePosition(tabController.tabPosition, stationaryTabController.tabPosition);
   }
-}
-
-/// The static map preview which moves around within the larger widget.
-class _MapPreview extends StatelessWidget {
-  const _MapPreview({Key? key, required this.size, required this.bShowShadow}) : super(key: key);
-
-  final Size size;
-  final bool bShowShadow;
 
   @override
   Widget build(BuildContext context) {
-    final borderRadius = BorderRadius.circular(6);
-
-    return SizedBox(
-      width: size.width,
-      height: size.height,
-      child: Stack(
-        children: [
-          AnimatedContainer(
-            clipBehavior: Clip.antiAlias,
-            duration: Duration(milliseconds: 400),
-            decoration: BoxDecoration(
-              borderRadius: borderRadius,
-              boxShadow: [
-                BoxShadow(
-                  color: Color(0xff000000).withOpacity(bShowShadow ? 0.4 : 0),
-                  spreadRadius: 2,
-                  blurRadius: 5,
-                ),
-              ],
-            ),
-            child: StageMap(
-              bIsControllable: false,
-              bShowOnlySelectedPins: true,
-              bForceShowFullMap: true,
-              pinSize: 32,
-            ),
-          ),
-          Container(
-            decoration: BoxDecoration(
-              borderRadius: borderRadius,
-              border: Border.all(
-                color: Color(0x50ffffff),
-                width: 0.5,
-              ),
-            ),
-          ),
-        ],
-      ),
-    );
-  }
-}
-
-/// The tab shown when the preview is docked.
-class _DockTab extends StatelessWidget {
-  const _DockTab(this.side, {Key? key}) : super(key: key);
-
-  /// Which side the preview is docked to.
-  final _DockSide side;
-
-  @override
-  Widget build(BuildContext context) {
-    const Radius borderRadius = Radius.circular(8);
-    final bool bIsFacingLeft = side == _DockSide.right;
-
-    return Container(
-      width: _dockTabSize.width,
-      height: _dockTabSize.height,
-      decoration: BoxDecoration(
-        color: Theme.of(context).colorScheme.surfaceTint,
-        borderRadius: bIsFacingLeft
-            ? BorderRadius.only(
-                topLeft: borderRadius,
-                bottomLeft: borderRadius,
-              )
-            : BorderRadius.only(
-                topRight: borderRadius,
-                bottomRight: borderRadius,
-              ),
-        boxShadow: [
-          BoxShadow(
-            color: Color(0x40000000),
-            offset: Offset(bIsFacingLeft ? -4 : 4, 4),
-            blurRadius: 4,
-          ),
-        ],
-      ),
-      child: Column(mainAxisAlignment: MainAxisAlignment.center, children: [
-        AssetIcon(
-          path: 'packages/epic_common/assets/icons/map.svg',
-          size: 24,
-        ),
-        SizedBox(
-          height: 6,
-        ),
-        AssetIcon(
-          path: bIsFacingLeft
-              ? 'packages/epic_common/assets/icons/chevron_left.svg'
-              : 'packages/epic_common/assets/icons/chevron_right.svg',
-          size: 24,
-        ),
-      ]),
+    return ChangeNotifierProvider.value(
+      value: tabController,
+      builder: (context, child) {
+        return FloatingTabBody(
+          AssetIcon(path: 'packages/epic_common/assets/icons/map.svg', size: 24),
+          StageMap(bIsControllable: false, bShowOnlySelectedPins: true, bForceShowFullMap: true, pinSize: 32),
+          bShouldUseStackLayout: true,
+          onTap: onTap,
+          onPanStart: onStartDrag,
+          onPanUpdate: updateDrag,
+          onPanCancel: endDrag,
+          onPanEnd: endDrag,
+          children: [],
+        );
+      },
     );
   }
 }
diff --git a/Engine/Extras/VirtualProduction/EpicStageApp/lib/widgets/elements/floating_tab_base.dart b/Engine/Extras/VirtualProduction/EpicStageApp/lib/widgets/elements/floating_tab_base.dart
new file mode 100644
--- /dev/null
+++ b/Engine/Extras/VirtualProduction/EpicStageApp/lib/widgets/elements/floating_tab_base.dart
@@ -0,0 +1,539 @@
+// Copyright Epic Games, Inc. All Rights Reserved.
+
+import 'package:epic_common/widgets.dart';
+import 'package:flutter/material.dart';
+import 'package:provider/provider.dart';
+
+import '../../models/unreal_dockable_tab_manager.dart';
+
+/// Sides of the screen on which the tab can be docked.
+enum DockSide {
+  left,
+  right,
+}
+
+/// Tabs Relative position.
+enum TabPosition { above, below }
+
+/// Possible states for the floating tab in terms of docking.
+enum DockState {
+  /// The tab is on-screen and no dock tabs are visible.
+  undocked,
+
+  /// The tab has entered the docking range and is revealing its tab.
+  docking,
+
+  /// The tab has revealed its docking tab. It may not actually be off-screen yet.
+  docked,
+
+  /// The tab has exited the docking range and is hiding its tab.
+  undocking,
+}
+
+/// The smallest dimension of the container will be multiplied by this value to determine the height of the tab.
+const double tabSizeScale = 0.25;
+
+/// The height of the tab will be multiplied by this value to determine its width.
+const double tabAspectRatio = 16.0 / 9.0;
+
+/// How far to inset the tab from the edge of the screen.
+const double edgeInset = 8.0;
+
+/// The size of the tab shown when the tab is docked.
+const dockTabSize = Size(32, 80);
+
+/// Map persistence prefix.
+const String mapPreferencePrefix = 'floatingMapSettings';
+
+/// Trackpad persistence prefix.
+const String trackpadPreferencePrefix = 'floatingTrackpadSettings';
+
+/// The visible part of the tab shown when it is docked (controller).
+class DockTab extends StatelessWidget {
+  const DockTab(this.side, this.indicator, {Key? key}) : super(key: key);
+
+  /// Which side the preview is docked to.
+  final DockSide side;
+
+  /// Indicator widget showing what the tab contains.
+  final Widget indicator;
+
+  @override
+  Widget build(BuildContext context) {
+    const Radius borderRadius = Radius.circular(8);
+    final bool bIsFacingLeft = side == DockSide.right;
+
+    return Container(
+      width: dockTabSize.width,
+      height: dockTabSize.height,
+      decoration: BoxDecoration(
+        color: Theme.of(context).colorScheme.surfaceTint,
+        borderRadius: bIsFacingLeft
+            ? BorderRadius.only(topLeft: borderRadius, bottomLeft: borderRadius)
+            : BorderRadius.only(topRight: borderRadius, bottomRight: borderRadius),
+        boxShadow: [
+          BoxShadow(
+            color: Color(0x40000000),
+            offset: Offset(bIsFacingLeft ? -4 : 4, 4),
+            blurRadius: 4,
+          ),
+        ],
+      ),
+      child: Column(mainAxisAlignment: MainAxisAlignment.center, children: [
+        indicator,
+        SizedBox(height: 6),
+        AssetIcon(
+          path: bIsFacingLeft
+              ? 'packages/epic_common/assets/icons/chevron_left.svg'
+              : 'packages/epic_common/assets/icons/chevron_right.svg',
+          size: 24,
+        ),
+      ]),
+    );
+  }
+}
+
+/// The static tab container which moves around within the larger widget.
+class FloatingTabContainer extends StatelessWidget {
+  const FloatingTabContainer({
+    Key? key,
+    required this.size,
+    required this.bShowShadow,
+    required this.child,
+    this.bUseStackedLayout = false,
+  }) : super(key: key);
+
+  /// size of the static container.
+  final Size size;
+
+  /// Whether the container has a shadow or not.
+  final bool bShowShadow;
+
+  /// Content of the container.
+  final Widget child;
+
+  /// Whether we want to use a stacked layout a more simpler layout.
+  final bool bUseStackedLayout;
+
+  @override
+  Widget build(BuildContext context) {
+    final borderRadius = BorderRadius.circular(6);
+
+    final Widget content = AnimatedContainer(
+      clipBehavior: Clip.antiAlias,
+      duration: Duration(milliseconds: 400),
+      decoration: BoxDecoration(
+        borderRadius: borderRadius,
+        border: Border.all(
+          color: Color(0x50ffffff),
+          width: 0.5,
+        ),
+        boxShadow: [
+          BoxShadow(
+            color: Color(0xff000000).withOpacity(bShowShadow ? 0.4 : 0),
+            spreadRadius: 2,
+            blurRadius: 5,
+          ),
+        ],
+      ),
+      child: child,
+    );
+
+    return SizedBox(
+      width: size.width,
+      height: size.height,
+      child: bUseStackedLayout
+          ? Stack(
+              children: [
+                content,
+                Container(
+                  decoration: BoxDecoration(
+                    borderRadius: borderRadius,
+                    border: Border.all(
+                      color: Color(0x50ffffff),
+                      width: 0.5,
+                    ),
+                  ),
+                ),
+              ],
+            )
+          : content,
+    );
+  }
+}
+
+/// Container for tab controls & content
+class FloatingTabBody extends StatelessWidget {
+  const FloatingTabBody(
+    this.icon,
+    this.child, {
+    super.key,
+    this.children = const [],
+    this.bIsDraggable = false,
+    this.bShouldUseStackLayout = false,
+    this.bCanInteract = false,
+    this.onTap,
+    this.onPanStart,
+    this.onPanUpdate,
+    this.onPanEnd,
+    this.onPanCancel,
+  });
+
+  /// List of widgets to be stacked above the [child] (main content) of the tab.
+  final List<Widget> children;
+
+  /// tab indicator icon.
+  final Widget icon;
+
+  /// whether or not the tab's position can be changed via dragging.
+  final bool bIsDraggable;
+
+  /// whether or not the tabs main content can be interacted with.
+  final bool bCanInteract;
+
+  /// main content of the tab.
+  final Widget child;
+
+  /// whether to use a stack layout in the [FloatingTabContainer] or not.
+  final bool bShouldUseStackLayout;
+
+  /// handles tap gestures on the tab body.
+  final VoidCallback? onTap;
+
+  /// handles pan start gestures.
+  final Function(DragStartDetails)? onPanStart;
+
+  /// handles pan update gestures.
+  final Function(DragUpdateDetails)? onPanUpdate;
+
+  /// handles pan end gestures.
+  final Function(Velocity)? onPanEnd;
+
+  /// handles pan cancel gestures.
+  final Function(Velocity)? onPanCancel;
+
+  @override
+  Widget build(BuildContext context) {
+    return LayoutBuilder(builder: (context, BoxConstraints constraints) {
+      final Size newSize = Size(constraints.maxWidth, constraints.maxHeight);
+
+      final _tabController = context.watch<UnrealFloatingTabController>();
+
+      // Handle first size/resize
+      if (newSize != _tabController.containerSize) {
+        final Size oldSize = _tabController.containerSize;
+        _tabController.containerSize = newSize;
+
+        if (oldSize == Size.zero) {
+          _tabController.initPosition();
+        } else {
+          WidgetsBinding.instance.addPostFrameCallback((_) => _tabController.onContainerResized(oldSize));
+        }
+      }
+
+      // Determine the offset of the main preview window, which shifts to the side to reveal dock tabs when docked.
+      double dockOffset = 0;
+      if (_tabController.bIsDockedOrDocking) {
+        switch (_tabController.dockSide) {
+          case DockSide.left:
+            dockOffset = -dockTabSize.width;
+            break;
+
+          case DockSide.right:
+            dockOffset = dockTabSize.width;
+            break;
+        }
+      }
+
+      return Stack(
+        clipBehavior: Clip.hardEdge,
+        children: [
+          // Position a box containing the preview, dock tabs, and gesture detector
+          Positioned(
+            top: _tabController.tabPosition.dy,
+            left: _tabController.tabPosition.dx,
+            child: SizedBox(
+              width: _tabController.tabSize.width,
+              height: _tabController.tabSize.height,
+              child: Stack(
+                clipBehavior: Clip.none,
+                children: [
+                  // Tab shown on the left when docked on the right
+                  if (_tabController.dockState != DockState.undocked && _tabController.dockSide == DockSide.right)
+                    Positioned(
+                      key: Key('LeftDockTab'),
+                      left: 0,
+                      top: (_tabController.tabSize.height - dockTabSize.height) / 2,
+                      child: DockTab(DockSide.right, icon),
+                    ),
+
+                  // Tab shown on the right when docked on the left
+                  if (_tabController.dockState != DockState.undocked && _tabController.dockSide == DockSide.left)
+                    Positioned(
+                      key: Key('RightDockTab'),
+                      left: _tabController.tabSize.width - dockTabSize.width,
+                      top: (_tabController.tabSize.height - dockTabSize.height) / 2,
+                      child: DockTab(DockSide.left, icon),
+                    ),
+
+                  // Preview window, which shifts to reveal tab when docked and hides when offscreen
+                  AnimatedPositioned(
+                    left: dockOffset,
+                    duration: const Duration(milliseconds: 200),
+                    child: _tabController.bIsTabVisible
+                        ? IgnorePointer(
+                            ignoring: bIsDraggable,
+                            child: FloatingTabContainer(
+                              size: _tabController.tabSize,
+                              bShowShadow: _tabController.bIsUndockedOrUndocking,
+                              child: child,
+                              bUseStackedLayout: bShouldUseStackLayout,
+                            ),
+                          )
+                        : SizedBox.shrink(),
+                    onEnd: _tabController.onDockTabAnimationEnd,
+                  ),
+
+                  // Gesture detector, which matches the preview rectangle when undocked and the revealed tab when docked
+                  Positioned(
+                    // When docked to left, offset to line up with the tab on the right side of the preview
+                    left: (_tabController.dockState == DockState.docked && _tabController.dockSide == DockSide.left)
+                        ? (_tabController.tabSize.width - dockTabSize.width)
+                        : 0,
+                    // When docked, offset to line up with the top of the tab
+                    top: _tabController.dockState == DockState.docked
+                        ? ((_tabController.tabSize.height - dockTabSize.height) / 2)
+                        : 0,
+                    child: SizedBox(
+                      // When docked, match size to the tab
+                      width: _tabController.dockState == DockState.docked
+                          ? dockTabSize.width
+                          : _tabController.tabSize.width,
+                      height: _tabController.dockState == DockState.docked
+                          ? dockTabSize.height
+                          : _tabController.tabSize.height,
+
+                      child: IgnorePointer(
+                        ignoring: bCanInteract,
+                        child: GestureDetector(
+                          onTap: onTap,
+                          onPanStart: onPanStart,
+                          onPanUpdate: onPanUpdate,
+                          onPanEnd: (details) => onPanEnd?.call(details.velocity),
+                          onPanCancel: () => onPanCancel?.call(Velocity.zero),
+                        ),
+                      ),
+                    ),
+                  ),
+
+                  ...children.map((e) => e).toList(),
+                ],
+              ),
+            ),
+          ),
+        ],
+      );
+    });
+  }
+}
+
+/// Base mixin to manage state of a floating tab.
+mixin FloatingTabBase<T extends StatefulWidget> on State<T> {
+  /// An instance of unreal tab manager.
+  late final UnrealDockableTabManager manager;
+
+  late final UnrealFloatingTabController tabController;
+  late final UnrealFloatingTabController stationaryTabController;
+
+  /// Relative position of the tab.
+  TabPosition tabRelativePosition = TabPosition.above;
+
+  @override
+  void initState() {
+    manager = Provider.of<UnrealDockableTabManager>(context, listen: false);
+    super.initState();
+  }
+
+  /// To be called from the [initState] of our user class, allows us to instantiate out stationary tab.
+  void initStationaryController(UnrealFloatingTabController controller) {
+    stationaryTabController = controller;
+  }
+
+  /// init tab controller should be called once but within the [initState] of our user class.
+  void initTab(String preferencePrefix) {
+    tabController = UnrealFloatingTabController(preferencePrefix);
+  }
+
+  /// Handle a tap on the preview.
+  void onTap() {
+    if (tabController.bIsUndockedOrUndocking) {
+      return;
+    }
+
+    // Currently docked, so undock by moving to the nearest corner.
+    tabController.moveTabToNearestCorner(tabController.targetTabPosition);
+    _moveStationaryTab();
+  }
+
+  /// Update the dragged preview position.
+  void updateDrag(DragUpdateDetails details) {
+    tabController.lastDragPosition = details.globalPosition;
+
+    tabController.moveTab(
+      tabController.targetTabPosition + Offset(details.delta.dx, details.delta.dy),
+      springDescription: tabController.draggedSpringDescription,
+    );
+
+    final DockSide? newDockSide = tabController.checkDockSideAtPosition(tabController.targetTabPosition);
+    if (newDockSide != null) {
+      // If this will dock, show the tab
+      tabController.dockTab(newDockSide, moveOffscreen: false);
+    } else {
+      // No longer docked, so move to undocking state
+      tabController.undockTab();
+    }
+  }
+
+  /// Stop dragging the preview and let it come to rest.
+  void endDrag(Velocity velocity) {
+    final RenderBox? renderBox = context.findRenderObject() as RenderBox?;
+    if (renderBox == null) {
+      tabController.moveTab(
+        Offset(edgeInset, edgeInset),
+        springDescription: tabController.releasedSpringDescription,
+      );
+      return;
+    }
+
+    // Check ahead from the preview's position so the user can "fling" the preview off the screen
+    const dockCheckAheadSeconds = 0.1;
+    final Offset dockCheckPosition = tabController.tabPosition + velocity.pixelsPerSecond * dockCheckAheadSeconds;
+
+    // If near the edges, dock the preview
+    final DockSide? newDockSide = tabController.checkDockSideAtPosition(dockCheckPosition);
+    if (newDockSide != null) {
+      tabController.dockTab(newDockSide, dockY: dockCheckPosition.dy);
+      _moveStationaryTab(dockCheckPosition);
+      return;
+    }
+
+    // Otherwise, move to the nearest corner. In this case, use the user's drag position instead of the preview
+    // position since flinging feels more accurate that way.
+    const cornerCheckAheadSeconds = 0.15;
+    final Offset cornerCheckPosition =
+        tabController.lastDragPosition + velocity.pixelsPerSecond * cornerCheckAheadSeconds;
+    final Offset finalLocalPosition = renderBox.globalToLocal(cornerCheckPosition);
+
+    tabController.moveTabToNearestCorner(finalLocalPosition);
+    _moveStationaryTab(finalLocalPosition);
+  }
+
+  /// Update spring preview position using springs.
+  void onTick(Duration elapsed) {
+    if (tabController.xSpring == null || tabController.ySpring == null) {
+      return;
+    }
+
+    tabController.lastSpringTime = elapsed.inMilliseconds / 1000.0;
+
+    final double springTime = tabController.lastSpringTime - tabController.lastSpringMoveTime;
+
+    setState(() {
+      tabController.tabVelocity = Offset(tabController.xSpring!.dx(springTime), tabController.ySpring!.dx(springTime));
+      tabController.tabPosition = Offset(tabController.xSpring!.x(springTime), tabController.ySpring!.x(springTime));
+    });
+
+    // If both springs are at rest, stop ticking to save performance
+    if (tabController.xSpring!.isDone(springTime) && tabController.ySpring!.isDone(springTime)) {
+      tabController.ticker?.stop();
+      tabController.lastSpringMoveTime = 0;
+      tabController.lastSpringTime = 0;
+    }
+  }
+
+  /// Start dragging the floating tab.
+  void onStartDrag(DragStartDetails details) {
+    tabRelativePosition =
+        manager.getTabRelativePosition(tabController.tabPosition, stationaryTabController.tabPosition);
+  }
+
+  /// tab avoidance logic on first app run only.
+  void tabAvoidanceOnFirstTimeRun() {
+    if (tabController.tabPosition.dy == stationaryTabController.tabPosition.dy) {
+      tabController.moveTabToFarthestCorner(
+        Offset(
+          tabController.tabPosition.dx,
+          tabController.tabPosition.dy - tabController.tabSize.height,
+        ),
+      );
+    }
+  }
+
+  /// Move a stationary tab out of the way for an incoming tab that's being dragged by providing the expected position
+  /// [position] of the dragged tab.
+  /// if [position] is not provided we'd use the current [draggedTabPosition] instead.
+  void _moveStationaryTab([Offset? position]) {
+    // Relative position of the dragged tab.
+
+    TabPosition relativePosition = tabRelativePosition;
+
+    final stationaryTabPosition = stationaryTabController.tabPosition;
+
+    final draggedTabPosition = position != null ? position : tabController.tabPosition;
+
+    double oneTabHeight = tabController.tabSize.height;
+
+    /// whether a fling would result in dragged tabs side being same with the stationary tabs side based off of the
+    /// tabs [tabController.targetPosition].
+    bool bIsDockSideToBeTheSame = tabController.checkDockSideAtPosition(position ?? tabController.targetTabPosition) ==
+        stationaryTabController.side;
+
+    void _moveTabBasedOnBottomAndTopEdges() {
+      // Center position of the stationary tab
+      double stationaryTabCenter = stationaryTabPosition.dy + oneTabHeight / 2;
+
+      // Bottom edge position of the currently being dragged tab
+      double draggedTabBottomEdge = (position != null ? position.dy : draggedTabPosition.dy) + oneTabHeight;
+
+      // Check if the drag started below the stationary tab and that the bottomLeft corner position is below the
+      // center position of the stationary tab.
+      if (relativePosition == TabPosition.above && draggedTabBottomEdge > stationaryTabCenter) {
+        double newPosition = stationaryTabPosition.dy - oneTabHeight;
+
+        stationaryTabController.dockTab(stationaryTabController.side, dockY: newPosition);
+
+        if (stationaryTabPosition.dy.round() == draggedTabPosition.dy.round()) {
+          stationaryTabController.moveTabToFarthestCorner(stationaryTabPosition);
+        }
+      } else if (relativePosition == TabPosition.below && draggedTabPosition.dy < stationaryTabCenter) {
+        double newPosition = stationaryTabPosition.dy + oneTabHeight;
+        stationaryTabController.dockTab(stationaryTabController.side, dockY: newPosition);
+
+        if (stationaryTabPosition.dy.round() == draggedTabPosition.dy.round()) {
+          stationaryTabController.moveTabToFarthestCorner(stationaryTabPosition);
+        }
+      }
+    }
+
+    void _handleEqualTabPositions() {
+      if (stationaryTabPosition.dy.round() == draggedTabPosition.dy.round()) {
+        stationaryTabController.moveTabToFarthestCorner(stationaryTabPosition);
+      }
+
+      if (stationaryTabPosition.dy.round() == tabController.targetTabPosition.dy.round()) {
+        stationaryTabController.moveTabToFarthestCorner(stationaryTabPosition);
+      }
+    }
+
+    // Check to make sure both tabs are on the same side
+    if (stationaryTabController.side == tabController.side) {
+      _moveTabBasedOnBottomAndTopEdges();
+      _handleEqualTabPositions();
+    }
+
+    if (bIsDockSideToBeTheSame) {
+      _handleEqualTabPositions();
+      _moveTabBasedOnBottomAndTopEdges();
+    }
+  }
+}
diff --git a/Engine/Extras/VirtualProduction/EpicStageApp/lib/widgets/elements/floating_trackpad.dart b/Engine/Extras/VirtualProduction/EpicStageApp/lib/widgets/elements/floating_trackpad.dart
new file mode 100644
--- /dev/null
+++ b/Engine/Extras/VirtualProduction/EpicStageApp/lib/widgets/elements/floating_trackpad.dart
@@ -0,0 +1,119 @@
+// Copyright Epic Games, Inc. All Rights Reserved.
+
+import 'dart:async';
+
+import 'package:epic_common/preferences.dart';
+import 'package:epic_common/theme.dart';
+import 'package:epic_common/widgets.dart';
+import 'package:flutter/material.dart';
+import 'package:provider/provider.dart';
+
+import '../../models/settings/floating_tab_settings.dart';
+import '../../models/settings/selected_actor_settings.dart';
+import 'floating_tab_base.dart';
+import 'lightcard_trackpad.dart';
+
+/// A floating preview of the stage map that the user can drag to any corner of the screen and dock to the side.
+class FloatingTrackpad extends StatefulWidget {
+  const FloatingTrackpad({Key? key}) : super(key: key);
+
+  static GlobalKey trackpad = GlobalKey(debugLabel: 'floating_trackpad');
+
+  @override
+  State<FloatingTrackpad> createState() => _FloatingTrackpadState();
+}
+
+class _FloatingTrackpadState extends State<FloatingTrackpad> with SingleTickerProviderStateMixin, FloatingTabBase {
+  /// Whether trackpad is locked for LC positioning input or unlocked for trackpad positioning drag.
+  bool _bIsTrackpadLocked = false;
+  late final SelectedActorSettings _actorSettings;
+
+  /// Instance of user base floating tab settings for trackpad.
+  late final FloatingTabSettings _settings;
+
+  /// Stream subscription for when the [tabController.dockState] changes.
+  late final StreamSubscription _subscription;
+
+  @override
+  void initState() {
+    super.initState();
+
+    tabController = manager.trackpadTabController;
+
+    initStationaryController(manager.mapTabController);
+
+    tabController.ticker = createTicker(onTick);
+
+    _actorSettings = Provider.of<SelectedActorSettings>(context, listen: false);
+
+    _settings = FloatingTabSettings(PreferencesBundle.of(context), trackpadPreferencePrefix);
+
+    /// init dragged tab relative position.
+    tabRelativePosition =
+        manager.getTabRelativePosition(tabController.tabPosition, stationaryTabController.tabPosition);
+
+    /// This helps us listen for changes to our dockState, if this tab is ever undocked while [_bIsTrackpadLocked] is
+    /// true, all drag and tap gestures would be ignored. to prevent this we want to know when the dock state changes
+    /// and set [_bIsTrackpadLocked] to false.
+    _subscription = _settings.bIsTabDocked.listen((value) => _handleOnDockStateChanged);
+  }
+
+  /// handle on dock state changed.
+  void _handleOnDockStateChanged(bool value) {
+    if (value && _bIsTrackpadLocked) {
+      setState(() => _bIsTrackpadLocked = false);
+    }
+  }
+
+  @override
+  void dispose() {
+    _subscription.cancel();
+    super.dispose();
+  }
+
+  @override
+  Widget build(BuildContext context) {
+    tabAvoidanceOnFirstTimeRun();
+    return ChangeNotifierProvider.value(
+      value: tabController,
+      builder: (context, child) {
+        return FloatingTabBody(
+          Icon(Icons.games_outlined),
+          LightCardTrackPad(),
+          bIsDraggable: !_bIsTrackpadLocked && _actorSettings.selectedActors.getValue().isNotEmpty,
+          bCanInteract: _bIsTrackpadLocked,
+          onTap: onTap,
+          onPanStart: onStartDrag,
+          onPanUpdate: _bIsTrackpadLocked ? null : updateDrag,
+          onPanCancel: endDrag,
+          onPanEnd: endDrag,
+          children: [
+            // Lock/unlock button which determines whether or not the trackpad is open modifying the position of a
+            // selected lightcard.
+            Visibility(
+              visible: tabController.dockState == DockState.undocked,
+              child: Padding(
+                padding: const EdgeInsets.all(8.0),
+                child: SizedBox(
+                  width: 40,
+                  height: 40,
+                  child: EpicIconButton(
+                    bIsToggledOn: _bIsTrackpadLocked,
+                    activeBackgroundColor: UnrealColors.gray18.withOpacity(.1),
+                    backgroundShape: BoxShape.circle,
+                    iconPath: _bIsTrackpadLocked
+                        ? 'packages/epic_common/assets/icons/lock_locked.svg'
+                        : 'packages/epic_common/assets/icons/lock_unlocked.svg',
+                    onPressed: () => setState(() {
+                      _bIsTrackpadLocked = !_bIsTrackpadLocked;
+                    }),
+                  ),
+                ),
+              ),
+            ),
+          ],
+        );
+      },
+    );
+  }
+}
diff --git a/Engine/Extras/VirtualProduction/EpicStageApp/lib/widgets/elements/lightcard_trackpad.dart b/Engine/Extras/VirtualProduction/EpicStageApp/lib/widgets/elements/lightcard_trackpad.dart
new file mode 100644
--- /dev/null
+++ b/Engine/Extras/VirtualProduction/EpicStageApp/lib/widgets/elements/lightcard_trackpad.dart
@@ -0,0 +1,304 @@
+//Copyright Epic Games, Inc. All Rights Reserved.
+
+import 'dart:async';
+
+import 'package:flutter/material.dart';
+import 'package:provider/provider.dart';
+
+import '../../models/actor_data/light_card_actor_data.dart';
+import '../../models/settings/delta_widget_settings.dart';
+import '../../models/unreal_property_controller.dart';
+import '../../models/settings/selected_actor_settings.dart';
+
+import '../../models/property_modify_operations.dart';
+import '../../models/unreal_actor_manager.dart';
+import '../../models/unreal_types.dart';
+import '../../utilities/constants.dart';
+
+enum _TrackpadOperation { negative, positive }
+
+/// Trackpad widget to moving and positioning LCs, works like an actual trackpad.
+class LightCardTrackPad extends StatefulWidget {
+  const LightCardTrackPad({super.key});
+
+  @override
+  State<LightCardTrackPad> createState() => _LightCardTrackPadState();
+}
+
+class _LightCardTrackPadState extends State<LightCardTrackPad> {
+  late final UnrealActorManager _actorManager;
+  late final SelectedActorSettings _selectedActorSettings;
+
+  late final DeltaWidgetSettings _deltaSettings;
+
+  /// Property Controller for latitude.
+  late final UnrealPropertyController<double> _latitudeController;
+
+  /// Property controller for longitude.
+  late final UnrealPropertyController<double> _longitudeController;
+
+  /// List containing whether or not the trackpad is reversed for each selected actor.
+  List<bool> _listBIsTrackpadReversed = [];
+
+  /// Gets a list of selected actor paths.
+  List<String> _getSelectedActorPaths() {
+    final List<String> positionedActorClasses = [lightCardClassName];
+    positionedActorClasses.addAll(colorCorrectWindowClassNames);
+
+    final List<String> nonUVActorPaths = [];
+
+    for (final String actorPath in _selectedActorSettings.selectedActors.getValue()) {
+      final UnrealObject? actor = _actorManager.getActorAtPath(actorPath);
+      if (!(actor?.isAny(positionedActorClasses) ?? false)) {
+        // Actor is not a class with position properties, so leave it out entirely
+        continue;
+      }
+
+      final LightCardActorData? lightCardActorData = actor!.getPerClassData<LightCardActorData>();
+      if (lightCardActorData?.bIsUV == false) {
+        nonUVActorPaths.add(actorPath);
+      }
+    }
+
+    return nonUVActorPaths;
+  }
+
+  /// Getter for list of selected actors path.
+  List<String> get _paths => _getSelectedActorPaths();
+
+  /// stream subscriptions for watching changes to selected actors.
+  StreamSubscription? _selectedActorsSubscription;
+
+  @override
+  void initState() {
+    _actorManager = Provider.of<UnrealActorManager>(context, listen: false);
+    _selectedActorSettings = Provider.of<SelectedActorSettings>(context, listen: false);
+    _deltaSettings = Provider.of<DeltaWidgetSettings>(context, listen: false);
+
+    _latitudeController = UnrealPropertyController(context, bShouldInitTransaction: false);
+    _longitudeController = UnrealPropertyController(context, bShouldInitTransaction: false);
+
+    _handleUpdateTrackedProperties();
+
+    _selectedActorsSubscription = _selectedActorSettings.selectedActors.listen((event) {
+      _handleUpdateTrackedProperties();
+    });
+
+    super.initState();
+  }
+
+  @override
+  void dispose() {
+    super.dispose();
+    _latitudeController.dispose();
+    _longitudeController.dispose();
+    _selectedActorsSubscription?.cancel();
+  }
+
+  @override
+  Widget build(BuildContext context) {
+    return GestureDetector(
+      onPanUpdate: _onPanUpdate,
+      onPanEnd: _onPanEnd,
+      onPanCancel: _onPanCancel,
+      onPanStart: _onPanStart,
+      child: Container(
+        decoration: BoxDecoration(
+          borderRadius: BorderRadius.circular(6),
+          color: Colors.black,
+        ),
+      ),
+    );
+  }
+
+  /// Handle vertical drag/pan across the trackpad widget providing Drag [details] containing positions on the x
+  /// and y axis.
+  void _onVerticalDrag(DragUpdateDetails details) {
+    List<double> summations = [];
+
+    double yAxis = details.delta.dy * _deltaSettings.sensitivity.getValue();
+
+    /// get actors latitude (property value) at [index].
+    double _getPropertyValue(int index) {
+      if (_latitudeController.makeValueDataList().isNotEmpty) {
+        return _latitudeController.makeValueDataList()[index]?.value ?? 0;
+      }
+      return 0;
+    }
+
+    /// increment the actors longitude at [index] with [value].
+    void _incrementLongitude(double value, int index) {
+      _modifyProperty(value, index, _longitudeController, PropertyMinMaxBehaviour.loop);
+    }
+
+    /// increment the actors longitude at [index] by 180° only when the actors latitude value at [index] is at 90° or
+    /// -90°.
+    void _incrementLongitudeByOneEighty(int index) {
+      if (_getPropertyValue(index) == 90 || _getPropertyValue(index) == -90) {
+        _incrementLongitude(180, index);
+      }
+    }
+
+    /// handle bounce back for when we hit the north pole or its opposite side, specifying a negative/positive
+    /// [operation] indicated whether the actors latitude is at 90° or -90° respectively, [summation]
+    /// which denotes the to be position on the Y-axis derived from adding the current y-position and the incoming
+    /// position from [details.delta.dy], and [index] which is the property index to be modified.
+    void _handleBounceBack(_TrackpadOperation operation, double summation, int index) {
+      // remainder from after we drive the latitude to 90° or -90°.
+      double overflow = operation == _TrackpadOperation.positive ? summation - 90 : -90 - summation;
+
+      // Drive the latitude to the max of 90° or -90°
+      _modifyProperty(yAxis - overflow, index, _latitudeController);
+
+      _incrementLongitudeByOneEighty(index);
+
+      // Drive the latitude in the reverse direction with the [overflow].
+      double tempValue = operation == _TrackpadOperation.positive ? -overflow * 2 : overflow * 2;
+
+      _modifyProperty(tempValue, index, _latitudeController);
+
+      // specify that we want our trackpad to now be reversed.
+      _listBIsTrackpadReversed[index] = true;
+    }
+
+    /// handle case for where the trackpad is marked as reversed.
+    void _handleReverseTrackpad(double summation, int index) {
+      bool bIsTrackpadReversed = _listBIsTrackpadReversed[index];
+
+      if (bIsTrackpadReversed) {
+        double currentPosition = _getPropertyValue(index);
+
+        if (summation < 90) {
+          _modifyProperty(-yAxis, index, _latitudeController);
+
+          if (currentPosition == 90) {
+            _incrementLongitudeByOneEighty(index);
+            _listBIsTrackpadReversed[index] = false;
+          }
+        }
+
+        if (currentPosition == -90) {
+          _incrementLongitudeByOneEighty(index);
+          _modifyProperty(yAxis, index, _latitudeController);
+
+          _listBIsTrackpadReversed[index] = false;
+        }
+
+        return;
+      }
+
+      _modifyProperty(yAxis, index, _latitudeController);
+    }
+
+    /// get (summations) to be positions of all actors on the y-axis by adding the value of [details.delta.dy] to an
+    /// actors current position.
+    for (int index = 0; index < _paths.length; index++) {
+      double latitude = 0;
+
+      /// check to make sure all selected actors have a valid position on the y-axis and that the length of paths
+      /// matches the length of data values for the actors.
+      if (_latitudeController.makeValueDataList().isNotEmpty &&
+          _paths.length == _latitudeController.makeValueDataList().length) {
+        latitude = _latitudeController.makeValueDataList()[index]?.value ?? 0;
+      }
+
+      /// sum of incoming dy [yAxis] and the current Y position [latitude] of LC'(s).
+      double summation = latitude + yAxis;
+
+      summations.add(summation);
+    }
+
+    /// handle when we hit either the negative or positive south pole.
+    for (int index = 0; index < summations.length; index++) {
+      double sum = summations[index];
+      if (sum > 90) {
+        _handleBounceBack(_TrackpadOperation.positive, sum, index);
+      } else if (sum < -90) {
+        _handleBounceBack(_TrackpadOperation.negative, sum, index);
+      } else {
+        _handleReverseTrackpad(sum, index);
+      }
+    }
+  }
+
+  /// Handle horizontal drag/pan across the trackpad widget providing Drag [details] containing positions on the x
+  /// and y axis.
+  void _onHorizontalDrag(DragUpdateDetails details) {
+    double dx = details.delta.dx * _deltaSettings.sensitivity.getValue();
+    for (var index = 0; index < _paths.length; index++) {
+      _modifyProperty(dx, index, _longitudeController, PropertyMinMaxBehaviour.loop);
+    }
+  }
+
+  /// Track all properties and update a property with [name] at the specified [index], modifying the property with
+  /// [value], using the related UnrealPropertyController [controller], optionally providing a
+  /// PropertyMinMaxBehaviour [behaviour] or [PropertyMinMaxBehaviour.clamp] will be used as default.
+  void _modifyProperty(double value, int index, UnrealPropertyController controller,
+      [PropertyMinMaxBehaviour behaviour = PropertyMinMaxBehaviour.clamp]) {
+    controller.modifyProperty(controller.modifyOperation, index, value: value, minMaxBehaviour: behaviour);
+  }
+
+  /// update tracked properties for both latitude and longitude controllers.
+  void _handleUpdateTrackedProperties() {
+    _latitudeController
+        .trackAllProperties(_getProperties('Latitude', _paths, modifierFn: _modifyPositionalPropertyNameBasedOnClass));
+    _longitudeController
+        .trackAllProperties(_getProperties('Longitude', _paths, modifierFn: _modifyPositionalPropertyNameBasedOnClass));
+  }
+
+  /// Handles all sort of drag/pan gestures on the trackpad, whether vertically or horizontally, providing Drag
+  /// update [details] which contains x [details.delta.dx] & y [details.delta.dy] positions.
+  void _onPanUpdate(DragUpdateDetails details) {
+    _longitudeController.beginTransaction();
+    _onVerticalDrag(details);
+    _onHorizontalDrag(details);
+  }
+
+  /// Called when pan/drag gestures end's, providing Drag end [details].
+  void _onPanEnd(DragEndDetails details) {
+    _longitudeController.endTransaction();
+    _listBIsTrackpadReversed = _paths.map((e) => false).toList();
+  }
+
+  /// Called when pan/drag gesture is canceled.
+  void _onPanCancel() {
+    _longitudeController.endTransaction();
+    _listBIsTrackpadReversed = _paths.map((e) => false).toList();
+  }
+
+  /// Called when pan/drag gesture starts.
+  void _onPanStart(DragStartDetails details) {
+    _listBIsTrackpadReversed = _paths.map((e) => false).toList();
+  }
+
+  /// Given an [actorPath] and a [propertyName], return a modified positional property name accounting for the actor's
+  /// type.
+  String _modifyPositionalPropertyNameBasedOnClass(String actorPath, String propertyName) {
+    final UnrealObject? actor = _actorManager.getActorAtPath(actorPath);
+
+    if (actor != null && actor.isAny(colorCorrectWindowClassNames)) {
+      // CCWs store their positional properties in a sub-structure.
+      return 'PositionalParams.' + propertyName;
+    }
+
+    return propertyName;
+  }
+}
+
+/// A simple type definition used by the [_geProperties] function.
+typedef ModifierFn = String Function(String path, String name);
+
+/// Get a list of properties with the given [name] for all of the actors with path in [paths].
+/// If [modifierFn] is provided, it will be called for each [actorPath] and the [propertyName] of the property,
+/// and its return value will be used in place of [name].
+List<UnrealProperty> _getProperties(String name, List<String> paths, {ModifierFn? modifierFn, String? overrideName}) {
+  return paths
+      .map(
+        (actorPath) => UnrealProperty(
+          objectPath: actorPath,
+          propertyName: (modifierFn != null) ? modifierFn(actorPath, name) : name,
+          typeNameOverride: overrideName,
+        ),
+      )
+      .toList();
+}
diff --git a/Engine/Extras/VirtualProduction/EpicStageApp/lib/widgets/elements/unreal_widget_base.dart b/Engine/Extras/VirtualProduction/EpicStageApp/lib/widgets/elements/unreal_widget_base.dart
--- a/Engine/Extras/VirtualProduction/EpicStageApp/lib/widgets/elements/unreal_widget_base.dart
+++ b/Engine/Extras/VirtualProduction/EpicStageApp/lib/widgets/elements/unreal_widget_base.dart
@@ -1,18 +1,13 @@
 // Copyright Epic Games, Inc. All Rights Reserved.
 
-import 'package:flutter/foundation.dart';
 import 'package:flutter/material.dart';
-import 'package:flutter_gen/gen_l10n/app_localizations.dart';
-import 'package:logging/logging.dart';
-import 'package:provider/provider.dart';
 
 import '../../models/property_modify_operations.dart';
 import '../../models/unreal_property_manager.dart';
+import '../../models/unreal_property_controller.dart';
 import '../../models/unreal_types.dart';
 import 'delta_widget_base.dart';
 
-final _log = Logger('UnrealWidget');
-
 /// Represents a single value shared by all controlled properties of a widget (if such a value exists).
 class SingleSharedValue<T> {
   const SingleSharedValue({required this.value, required this.bHasMultipleValues});
@@ -108,91 +103,45 @@ abstract class UnrealWidget extends StatefulWidget {
 /// - Call [handleOnResetByUser] when the user presses the widget's reset button.
 /// - Override [modifyOperation] if needed.
 mixin UnrealWidgetStateMixin<WidgetType extends UnrealWidget, PropertyType> on State<WidgetType> {
-  /// The property manager this widget uses to update its properties.
-  late final UnrealPropertyManager _propertyManager;
-
-  /// A list of properties currently controlled by the widget.
-  final List<WidgetControlledUnrealProperty<PropertyType>?> _properties = [];
-
-  /// A list of properties for enabling widget properties current controlled by the widget.
-  final List<WidgetControlledUnrealProperty<bool>?> _enableProperties = [];
-
-  /// The last list of properties that were passed to the widget. This is stored so we can compare the list on rebuild
-  /// and register/unregister properties when the list changes.
-  final List<UnrealProperty> _lastUnrealProperties = [];
-
-  /// The last list of enable properties that were passed to the widget. This is stored so we can compare the list on
-  /// rebuild and register/unregister properties when the list changes.
-  final List<UnrealProperty> _lastEnableProperties = [];
-
-  /// The property's metadata as received from the engine.
-  UnrealPropertyMetadata? _propertyMetadata;
-
-  /// This is incremented whenever the set of tracked properties changes so we can ignore any callbacks from previous
-  /// versions. For example, if the user rapidly changes the list multiple times, we may get changed value or newly
-  /// tracked value callbacks for properties that we no longer care about.
-  int _latestTrackedPropertyVersion = 0;
-
-  /// Whether the properties' values are ready to display to the user.
-  bool _bIsReady = false;
-
-  /// Whether the widget is actively in a transaction.
-  bool _bIsInTransaction = false;
-
-  /// Whether the properties' values are ready to display to the user.
-  bool get bIsReady => _bIsReady;
-
-  /// The property's metadata as received from the engine.
-  UnrealPropertyMetadata? get propertyMetadata => _propertyMetadata;
+  late final UnrealPropertyController<PropertyType> _singlePropertyController;
 
   /// The label of the property to show to users.
-  String get propertyLabel => widget.overrideName ?? (propertyMetadata?.displayName ?? '');
+  String get propertyLabel => widget.overrideName ?? (_singlePropertyController.propertyLabel);
 
   /// The property's minimum value as received from the engine.
-  PropertyType? get engineMin => propertyMetadata?.minValue;
+  PropertyType? get engineMin => _singlePropertyController.engineMin;
 
   /// The property's maximum value as received from the engine.
-  PropertyType? get engineMax => propertyMetadata?.maxValue;
+  PropertyType? get engineMax => _singlePropertyController.engineMax;
 
   /// The overridden minimum value the controlled properties can reach. If null, [engineMin] will be used.
-  PropertyType? get overrideMin => null;
+  PropertyType? get overrideMin => _singlePropertyController.overrideMin;
 
   /// The overridden maximum value the controlled properties can reach. If null, [engineMax] will be used.
-  PropertyType? get overrideMax => null;
+  PropertyType? get overrideMax => _singlePropertyController.overrideMax;
 
-  /// The property's possible enum values as received from the engine.
-  List<String> get propertyEnumValues => propertyMetadata?.enumValues ?? [];
+  /// The type of modify operation this widget uses to apply deltas.
+  PropertyModifyOperation get modifyOperation => _singlePropertyController.modifyOperation;
 
   /// The number of properties this is controlling.
-  int get propertyCount => _properties.length;
-
-  /// The type of modify operation this widget uses to apply deltas.
-  PropertyModifyOperation get modifyOperation => const AddOperation();
+  int get propertyCount => _singlePropertyController.properties.length;
 
-  /// Extra meta-data properties to pass to conversion functions operating on this widget's value.
-  Map<String, dynamic>? get conversionMetadata => null;
+  /// The property's possible enum values as received from the engine.
+  List<String> get propertyEnumValues => _singlePropertyController.propertyEnumValues;
 
   @override
   void initState() {
-    super.initState();
+    _singlePropertyController = UnrealPropertyController(context);
+    _singlePropertyController.trackAllProperties(widget.unrealProperties, widget.enableProperties);
+    _singlePropertyController.addListener(_handleOnPropertiesChanged);
 
-    _propertyManager = Provider.of<UnrealPropertyManager>(context, listen: false);
-    _trackAllProperties();
+    super.initState();
   }
 
   @override
   void dispose() {
-    for (final WidgetControlledUnrealProperty<PropertyType>? property in _properties) {
-      if (property != null) {
-        _propertyManager.stopTrackingProperty(property.trackedId, property.changedCallback);
-      }
-    }
-
-    for (final WidgetControlledUnrealProperty<bool>? property in _enableProperties) {
-      if (property != null) {
-        _propertyManager.stopTrackingProperty(property.trackedId, property.changedCallback);
-      }
-    }
+    _singlePropertyController.removeListener(_handleOnPropertiesChanged);
+    _singlePropertyController.dispose();
 
     super.dispose();
   }
@@ -200,8 +149,7 @@ mixin UnrealWidgetStateMixin<WidgetType extends UnrealWidget, PropertyType> on S
   @override
   void didUpdateWidget(WidgetType oldWidget) {
     super.didUpdateWidget(oldWidget);
-
-    _trackAllPropertiesIfChanged();
+    _singlePropertyController.trackAllPropertiesIfChanged(widget.unrealProperties, widget.enableProperties);
   }
 
   /// Call this to indicate that the user has started interacting with the widget.
@@ -209,67 +157,24 @@ mixin UnrealWidgetStateMixin<WidgetType extends UnrealWidget, PropertyType> on S
   /// A [description] will be automatically generated if null is passed.
   /// Returns true if a new transaction was started.
   bool beginTransaction([String? description]) {
-    if (_bIsInTransaction) {
-      return false;
-    }
-
-    if (_propertyManager
-        .beginTransaction(description ?? AppLocalizations.of(context)!.transactionEditProperty(propertyLabel))) {
-      _bIsInTransaction = true;
-      return true;
-    }
-
-    return false;
+    return _singlePropertyController.beginTransaction(description);
   }
 
   /// Call this to indicate that the user has finished interacting with the widget.
   void endTransaction() {
-    if (!_bIsInTransaction) {
-      return;
-    }
-
-    _bIsInTransaction = false;
-    _propertyManager.endTransaction();
+    _singlePropertyController.endTransaction();
   }
 
   /// Generate a list of [DeltaWidgetValueData] containing the data needed to display each control in the widget.
   /// Values may be null if the value hasn't been received from the engine yet.
-  List<DeltaWidgetValueData<PropertyType>?> makeValueDataList() =>
-      _bIsReady ? _properties.map((property) => property?.toValueData()).toList() : [];
-
-  /// Called when the user changes controlled values.
-  void handleOnChangedByUser(List<PropertyType> deltaValues) {
-    modifyProperties(modifyOperation, values: deltaValues);
-  }
-
-  /// Called when the user presses the reset button.
-  void handleOnResetByUser() {
-    if (!_propertyManager.beginTransaction(AppLocalizations.of(context)!.transactionResetProperty(propertyLabel))) {
-      return;
-    }
-
-    for (final WidgetControlledUnrealProperty<PropertyType>? property in _properties) {
-      if (property != null) {
-        _propertyManager.modifyTrackedPropertyValue(property.trackedId, const ResetOperation());
-      }
-    }
-
-    if (widget.enableProperties != null) {
-      for (final WidgetControlledUnrealProperty<bool>? property in _enableProperties) {
-        if (property != null) {
-          _propertyManager.modifyTrackedPropertyValue(property.trackedId, const ResetOperation());
-        }
-      }
-    }
-
-    _propertyManager.endTransaction();
-  }
+  List<DeltaWidgetValueData<PropertyType>?> makeValueDataList() => _singlePropertyController.makeValueDataList();
 
   /// Return a SingleSharedValue indicating whether there is a single value shared by all controlled properties.
   /// If there are no controlled values, the SingleSharedValue will instead contain [defaultValue].
   SingleSharedValue<PropertyType> getSingleSharedValue([PropertyType? defaultValue]) {
     // Get the value of all non-null properties
-    final List<PropertyType?> controlledValues = _properties.map((property) => property?.value).toList();
+    final List<PropertyType?> controlledValues =
+        _singlePropertyController.properties.map((property) => property?.value).toList();
     controlledValues.removeWhere((element) => element == null);
 
     bool bHasMultipleValues = false;
@@ -297,201 +202,31 @@ mixin UnrealWidgetStateMixin<WidgetType extends UnrealWidget, PropertyType> on S
     return SingleSharedValue(value: sharedValue, bHasMultipleValues: bHasMultipleValues);
   }
 
-  /// Apply an [operation] to each property this controls by the amounts specified in [values].
-  /// If [values] is null, this is treated as if it was a list of null values matching the number of properties.
-  /// If [bIgnoreLimits] is true, the min and max for the widget will be ignored.
-  void modifyProperties(PropertyModifyOperation operation, {List<dynamic>? values, bool bIgnoreLimits = false}) {
-    assert(values == null || values.length == _properties.length);
-
-    if (!_bIsInTransaction) {
-      if (!beginTransaction()) {
-        // If we can't start a transaction, ignore user input or we'll get out of sync with the engine
-        return;
-      }
-    }
-
-    // Enable the property in the engine
-    for (final WidgetControlledUnrealProperty<bool>? property in _enableProperties) {
-      if (property != null) {
-        if (!_propertyManager.getTrackedPropertyValue(property.trackedId)) {
-          _propertyManager.modifyTrackedPropertyValue(property.trackedId, const SetOperation(), deltaValue: true);
-        }
-      }
-    }
-
-    for (int propertyIndex = 0; propertyIndex < _properties.length; ++propertyIndex) {
-      final WidgetControlledUnrealProperty<PropertyType>? property = _properties[propertyIndex];
-
-      if (property != null) {
-        _propertyManager.modifyTrackedPropertyValue(
-          property.trackedId,
-          operation,
-          deltaValue: values?[propertyIndex],
-          minMaxBehaviour: bIgnoreLimits ? PropertyMinMaxBehaviour.ignore : widget.minMaxBehaviour,
-          overrideMin: overrideMin,
-          overrideMax: overrideMax,
-        );
-      }
-    }
-
-    if (widget.onChangedByUser != null) {
-      widget.onChangedByUser!();
-    }
-  }
-
-  /// Check whether the widget's list of properties has changed, and if so, update our list of tracked properties.
-  void _trackAllPropertiesIfChanged() {
-    bool bHavePropertiesChanged = widget.unrealProperties.length != _lastUnrealProperties.length ||
-        (widget.enableProperties?.length ?? 0) != _lastEnableProperties.length;
-
-    if (!bHavePropertiesChanged) {
-      bHavePropertiesChanged = !listEquals(widget.unrealProperties, _lastUnrealProperties);
-    }
-
-    if (!bHavePropertiesChanged && widget.enableProperties != null) {
-      bHavePropertiesChanged = !listEquals(widget.enableProperties!, _lastEnableProperties);
-    }
-
-    if (!bHavePropertiesChanged) {
-      return;
-    }
-
-    _trackAllProperties();
-  }
-
-  /// Resets the internal list of properties and starts tracking them again.
-  void _trackAllProperties() {
-    _bIsReady = false;
-    ++_latestTrackedPropertyVersion;
-
-    // The tracked property version when we started this function, so we can ignore callbacks later in the function
-    // if they come back too late to be relevant.
-    final int trackedPropertyVersion = _latestTrackedPropertyVersion;
-
-    // Futures that will complete when each property is exposed.
-    final List<Future<void>> exposeFutures = [];
-
-    _trackRelevantProperties<PropertyType>(
-      trackedProperties: _properties,
-      oldProperties: _lastUnrealProperties,
-      newProperties: widget.unrealProperties,
-      exposeFutures: exposeFutures,
-      callback: _handleOnManagedValueChanged,
-    );
-
-    _trackRelevantProperties<bool>(
-      trackedProperties: _enableProperties,
-      oldProperties: _lastEnableProperties,
-      newProperties: widget.enableProperties ?? [],
-      exposeFutures: exposeFutures,
-      callback: _handleOnEnableValueChanged,
-    );
-
-    // Wait for all properties to be subscribed, then mark the widget as ready to be used
-    Future.wait(exposeFutures).then((_) {
-      if (_latestTrackedPropertyVersion != trackedPropertyVersion || !mounted) {
-        return;
-      }
-
-      _onPropertiesReady();
-    }).onError((error, stackTrace) {
-      _log.severe('Error while waiting for Unreal widget properties', error, stackTrace);
-    });
+  /// handle state change for [properties].
+  void _handleOnPropertiesChanged() {
+    setState(() {});
   }
 
-  /// Given a list of [trackedProperties], a corresponding list of [oldProperties] that are being tracked, and a list
-  /// of [newProperties] that are about to be tracked, stop tracking any properties that are no longer relevant and
-  /// start any properties that are newly relevant. [exposeFutures] will be filled with futures for each newly-tracked
-  /// property that will complete when the property is exposed, and [callback] will be called in the future whenever a
-  /// property's value changes.
-  void _trackRelevantProperties<T>({
-    required List<WidgetControlledUnrealProperty<T>?> trackedProperties,
-    required List<UnrealProperty> oldProperties,
-    required List<UnrealProperty> newProperties,
-    required List<Future<void>> exposeFutures,
-    required Function(int trackedPropertyVersion, int index, dynamic newValue) callback,
-  }) {
-    assert(oldProperties.length == trackedProperties.length);
-
-    // Stop tracking all properties. We either aren't tracking them anymore, or we may need to update their
-    // propertyIndex if the number tracked has changed size, in which case we need to recreate the callback anyway.
-    for (int propertyIndex = 0; propertyIndex < oldProperties.length; ++propertyIndex) {
-      final WidgetControlledUnrealProperty<T>? property = trackedProperties[propertyIndex];
-
-      if (property != null) {
-        _propertyManager.stopTrackingProperty(property.trackedId, property.changedCallback);
-      }
-    }
-
-    oldProperties.clear();
-    oldProperties.addAll(newProperties);
-
-    trackedProperties.length = newProperties.length;
-
-    // Track each property we want to control with the remote property manager.
-    for (int propertyIndex = 0; propertyIndex < newProperties.length; ++propertyIndex) {
-      final UnrealProperty property = newProperties[propertyIndex];
-
-      // Create a callback function for this property
-      changedCallback(newValue) => callback(_latestTrackedPropertyVersion, propertyIndex, newValue);
-
-      // Start tracking the property in the property manager
-      final TrackedPropertyId trackedId =
-          _propertyManager.trackProperty(property, changedCallback, conversionMetadata: conversionMetadata);
-
-      final newProperty = WidgetControlledUnrealProperty<T>(
-        trackedId: trackedId,
-        propertyIndex: propertyIndex,
-        changedCallback: changedCallback,
-        value: _propertyManager.getTrackedPropertyValue(trackedId),
-        ownerName: property.objectName,
-      );
-      trackedProperties[propertyIndex] = newProperty;
-
-      // Add to the list of properties we're waiting for
-      final Future<void> exposeFuture = _propertyManager.waitForProperty(trackedId);
-      exposeFutures.add(exposeFuture);
-
-      exposeFuture.then((_) {
-        newProperty.value = _propertyManager.getTrackedPropertyValue(trackedId);
-      }).onError((error, stackTrace) {
-        _log.severe(
-            'Error while exposing property ${property.propertyName} on ${property.objectName}', error, stackTrace);
-      });
-    }
-  }
-
-  /// Called when all properties have been subscribed and the widget is ready to be interacted with.
-  void _onPropertiesReady() {
-    if (_properties.isNotEmpty) {
-      // Properties could differ in their metadata, but for now we'll assume they share the same values
-      final WidgetControlledUnrealProperty<PropertyType>? firstProperty = _properties[0];
-
-      if (firstProperty == null) {
-        assert(false, 'If an Unreal widget is ready, all properties should be non-null');
-        return;
-      }
-
-      _propertyMetadata = _propertyManager.getTrackedPropertyMetaData(firstProperty.trackedId);
-    }
-
+  /// Called when the user changes controlled values.
+  void handleOnChangedByUser(List<PropertyType> deltaValues) {
     setState(() {
-      _bIsReady = true;
+      _singlePropertyController.modifyProperties(
+        modifyOperation,
+        values: deltaValues,
+        onChangedByUser: widget.onChangedByUser,
+      );
     });
   }
 
-  /// Called when the property manager reports a [newValue] for the property with the given [index] in the widget.
-  /// If [trackedPropertyVersion] does not match [_latestTrackedPropertyVersion], the change will be ignored.
-  void _handleOnManagedValueChanged(int trackedPropertyVersion, int index, dynamic newValue) {
-    if (trackedPropertyVersion != _latestTrackedPropertyVersion) {
-      return;
-    }
-
-    setState(() {
-      _properties[index]?.value = newValue;
-    });
+  /// Called when the user presses the reset button.
+  void handleOnResetByUser() {
+    setState(() => _singlePropertyController.handleOnResetByUser(widget.enableProperties));
   }
 
-  /// Called when an enable value for one of the properties changes.
-  void _handleOnEnableValueChanged(int trackedPropertyVersion, int index, dynamic newValue) {}
+  /// Apply an [operation] to each property this controls by the amounts specified in [values].
+  /// If [values] is null, this is treated as if it was a list of null values matching the number of properties.
+  /// If [bIgnoreLimits] is true, the min and max for the widget will be ignored.
+  void modifyProperties(PropertyModifyOperation operation, {List<dynamic>? values, bool bIgnoreLimits = false}) {
+    _singlePropertyController.modifyProperties(operation, values: values, bIgnoreLimits: bIgnoreLimits);
+  }
 }
diff --git a/Engine/Extras/VirtualProduction/EpicStageApp/lib/widgets/screens/main/stage_app_main_screen.dart b/Engine/Extras/VirtualProduction/EpicStageApp/lib/widgets/screens/main/stage_app_main_screen.dart
--- a/Engine/Extras/VirtualProduction/EpicStageApp/lib/widgets/screens/main/stage_app_main_screen.dart
+++ b/Engine/Extras/VirtualProduction/EpicStageApp/lib/widgets/screens/main/stage_app_main_screen.dart
@@ -13,6 +13,7 @@ import '../../../models/unreal_types.dart';
 import '../../../utilities/constants.dart';
 import '../../../utilities/debug_utilities.dart';
 import '../../elements/floating_map_preview.dart';
+import '../../elements/floating_trackpad.dart';
 import '../connect/views/n_display_selector_dialog.dart';
 import 'tabs/main_screen_tabs.dart';
 import 'toolbar/main_screen_toolbar.dart';
@@ -123,6 +124,7 @@ class _StageAppMainScreenState extends State<StageAppMainScreen>
               ),
             ),
             if (_bShowMapPreview) const FloatingMapPreview(key: Key('Map Preview')),
+            FloatingTrackpad()
           ],
         ),
       ),
