[FlutterCommon] [UnrealStage] [LiveLinkVCAM] Show version info in settings menu and logs

diff --git a/Engine/Extras/FlutterCommon/pubspec.yaml b/Engine/Extras/FlutterCommon/pubspec.yaml
--- a/Engine/Extras/FlutterCommon/pubspec.yaml
+++ b/Engine/Extras/FlutterCommon/pubspec.yaml
@@ -18,6 +18,7 @@ dependencies:
   intl: any
   logging: ^1.2.0
   path_provider: ^2.0.11
+  package_info_plus: ^5.0.1
   provider: ^6.0.2
   streaming_shared_preferences: ^2.0.0
   uuid: ^4.2.1
diff --git a/Engine/Extras/VirtualProduction/EpicStageApp/lib/main.dart b/Engine/Extras/VirtualProduction/EpicStageApp/lib/main.dart
--- a/Engine/Extras/VirtualProduction/EpicStageApp/lib/main.dart
+++ b/Engine/Extras/VirtualProduction/EpicStageApp/lib/main.dart
@@ -7,6 +7,7 @@ import 'package:epic_common/localizations.dart';
 import 'package:epic_common/logging.dart';
 import 'package:epic_common/preferences.dart';
 import 'package:epic_common/theme.dart';
+import 'package:epic_common/utilities/version.dart';
 import 'package:epic_common/widgets.dart';
 import 'package:flutter/material.dart';
 import 'package:flutter/services.dart';
@@ -54,6 +55,8 @@ void main() async {
 void initAndRunApp() async {
   WidgetsFlutterBinding.ensureInitialized();
 
+  _log.info(await getVerbosePackageVersion());
+
   final flutterLog = Logger('Flutter');
   final isolateLog = Logger('Isolate');
 
diff --git a/Engine/Extras/VirtualProduction/EpicStageApp/lib/widgets/screens/main/toolbar/settings/pages/settings_dialog_main.dart b/Engine/Extras/VirtualProduction/EpicStageApp/lib/widgets/screens/main/toolbar/settings/pages/settings_dialog_main.dart
--- a/Engine/Extras/VirtualProduction/EpicStageApp/lib/widgets/screens/main/toolbar/settings/pages/settings_dialog_main.dart
+++ b/Engine/Extras/VirtualProduction/EpicStageApp/lib/widgets/screens/main/toolbar/settings/pages/settings_dialog_main.dart
@@ -1,6 +1,7 @@
 // Copyright Epic Games, Inc. All Rights Reserved.
 
 import 'package:epic_common/theme.dart';
+import 'package:epic_common/utilities/version.dart';
 import 'package:epic_common/widgets.dart';
 import 'package:flutter/material.dart';
 import 'package:flutter_gen/gen_l10n/app_localizations.dart';
@@ -71,14 +72,24 @@ class SettingsDialogMain extends StatelessWidget {
             iconPath: 'packages/epic_common/assets/icons/info.svg',
             onTap: () => rootNavigatorKey.currentState?.pushNamed(EulaScreen.route, arguments: {'onPressed': () {}}),
           ),
-          if (bIsConnected) const SettingsMenuDivider(),
-          if (bIsConnected)
-            Container(
-              padding: EdgeInsets.all(8),
-              child: Row(
-                children: const [_DisconnectButton()],
-              ),
+          const SettingsMenuDivider(),
+          Container(
+            padding: EdgeInsets.symmetric(horizontal: 8),
+            height: 54,
+            child: Row(
+              children: [
+                if (bIsConnected) const _DisconnectButton(),
+                const Spacer(),
+                FutureBuilder<String>(
+                  future: getFriendlyPackageVersion(),
+                  builder: (context, snapshot) => Text(
+                    snapshot.data ?? '',
+                    style: Theme.of(context).textTheme.labelMedium,
+                  ),
+                ),
+              ],
             ),
+          ),
         ],
       ),
     );
