Fix CPU/memory leak that would happen each time a session was connected to.

- Reconnections through the app were leaking ~200mb on each connection, this is now fixed.
- Fix reconnection modal not hiding on successful reconnection in some cases.
- Many small leaks where self was captured implicitly inside closures now use weak self.
- In app stats graphs now scale to different device resolutions
- Added pauseCount metric to the stats graphs
- Added ARKit metrics to the stats graphs

diff --git a/Engine/Extras/VirtualProduction/LiveLinkVCAM/vcam/HeaderView.swift b/Engine/Extras/VirtualProduction/LiveLinkVCAM/vcam/HeaderView.swift
--- a/Engine/Extras/VirtualProduction/LiveLinkVCAM/vcam/HeaderView.swift
+++ b/Engine/Extras/VirtualProduction/LiveLinkVCAM/vcam/HeaderView.swift
@@ -18,7 +18,8 @@ import QuickLayout
 class HeaderView : UIView {
     
     @objc dynamic let appSettings = AppSettings.shared
-    private var observers = [NSKeyValueObservation]()
+    private var livelinkSubjectNameObserver : NSKeyValueObservation?
+    private var timecodeSourceObserver : NSKeyValueObservation?
 
     private var displayLink: CADisplayLink!
     
@@ -55,8 +56,12 @@ class HeaderView : UIView {
         privateInit()
     }
     
+    deinit {
+        Log.info("HeaderView destructed.")
+    }
+    
     private func privateInit() {
-
+        Log.info("HeaderView constructed.")
     }
     
     override func awakeFromNib() {
@@ -167,37 +172,14 @@ class HeaderView : UIView {
             
             layer.addSublayer(border)
         }
-
-
-        observers.append(observe(\.appSettings.liveLinkSubjectName, options: [.initial,.new], changeHandler: { object, change in
-            self.subject = self.appSettings.liveLinkSubjectName
-        }))
-        
-        observers.append(observe(\.appSettings.timecodeSource, options: [.initial,.new], changeHandler: { object, change in
-            
-            switch self.appSettings.timecodeSourceEnum() {
-            case .systemTime:
-                self.timecodeImageView.image = UIImage(named: "counterDevice")
-            case .ntp:
-                self.timecodeImageView.image = UIImage(named: "counterCloud")
-            case .tentacleSync:
-                self.timecodeImageView.image = UIImage(named: "tentacle")
-            default:
-                break
-            }
-        }))
         
-        observers.append(observe(\.appSettings.showStreamingStats, options: [.initial,.new], changeHandler: { object, change in
-            
-        }))
     }
     
     func start() {
         // display link updates the timecode value
-        displayLink = CADisplayLink(
-          target: self, selector: #selector(displayLinkDidFire)
-        )
+        displayLink = CADisplayLink(target: self, selector: #selector(displayLinkDidFire))
         displayLink.add(to: .main, forMode: .common)
+        self.setupObservers()
     }
     
     func stop() {
@@ -205,6 +187,28 @@ class HeaderView : UIView {
         displayLink = nil
     }
     
+    func setupObservers() {
+        self.livelinkSubjectNameObserver = observe(\.appSettings.liveLinkSubjectName, options: [.initial,.new], changeHandler: { [weak self] object, change in
+            if let validSelf = self {
+                validSelf.subject = validSelf.appSettings.liveLinkSubjectName
+            }
+        })
+        
+        self.timecodeSourceObserver = observe(\.appSettings.timecodeSource, options: [.initial,.new], changeHandler: { [weak self] object, change in
+            
+            switch self?.appSettings.timecodeSourceEnum() {
+            case .systemTime:
+                self?.timecodeImageView.image = UIImage(named: "counterDevice")
+            case .ntp:
+                self?.timecodeImageView.image = UIImage(named: "counterCloud")
+            case .tentacleSync:
+                self?.timecodeImageView.image = UIImage(named: "tentacle")
+            default:
+                break
+            }
+        })
+    }
+    
     @objc func exitTapped(_ sender : Any) {
         if let d = self.delegate {
             d.headerViewExitButtonTapped(self)
@@ -223,9 +227,7 @@ class HeaderView : UIView {
         }
     }
 
-
     @objc func displayLinkDidFire(_ displayLink: CADisplayLink) {
-
         self.timecodeLabel.text = Timecode.create().toString(includeFractional: false)
     }
 }
diff --git a/Engine/Extras/VirtualProduction/LiveLinkVCAM/vcam/LineGraphView.swift b/Engine/Extras/VirtualProduction/LiveLinkVCAM/vcam/LineGraphView.swift
--- a/Engine/Extras/VirtualProduction/LiveLinkVCAM/vcam/LineGraphView.swift
+++ b/Engine/Extras/VirtualProduction/LiveLinkVCAM/vcam/LineGraphView.swift
@@ -15,28 +15,28 @@ class LineGraphView: UIView {
     private var dataPoints: [CGFloat] = []
     
     // Limit to the number of data points stored in the array
-    private var sizeLimit: UInt16 = 120;
+    private var sizeLimit: UInt16 = 120
     
     // Font to use in our graph title
-    private var titleFont: UIFont = UIFont.systemFont(ofSize: 12);
+    private var titleFont: UIFont
     
     // Font to use in our graph tick labels
-    private var tickLabelFont: UIFont = UIFont.systemFont(ofSize: 8);
+    private var tickLabelFont: UIFont
     
     // Color of our font
-    private var fontColor: UIColor = UIColor.white;
+    private var fontColor: UIColor = UIColor.white
     
     // Paragraph style for our text
-    private var paragraphStyle: NSMutableParagraphStyle = NSMutableParagraphStyle();
+    private var paragraphStyle: NSMutableParagraphStyle = NSMutableParagraphStyle()
     
     // Title text attributes
-    private var titleTextAttributes: [NSAttributedString.Key: Any] = [:];
+    private var titleTextAttributes: [NSAttributedString.Key: Any] = [:]
     
     // Tick label text attributes
-    private var tickTextAttributes: [NSAttributedString.Key: Any] = [:];
+    private var tickTextAttributes: [NSAttributedString.Key: Any] = [:]
     
     // Maximum y value we have seen
-    public var maxY: CGFloat = 1;
+    public var maxY: CGFloat = 1
     
     // Title of graph
     public var title: String = "Title" {
@@ -51,6 +51,12 @@ class LineGraphView: UIView {
     }
     
     override init(frame: CGRect) {
+        // Font to use in our graph title
+        self.titleFont = UIFont.systemFont(ofSize: 12 / (UIScreen.main.scale - 1));
+        
+        // Font to use in our graph tick labels
+        self.tickLabelFont = UIFont.systemFont(ofSize: 8 / (UIScreen.main.scale - 1));
+        
         super.init(frame: frame)
         
         // set transparent background color
@@ -79,6 +85,12 @@ class LineGraphView: UIView {
     }
     
     required init?(coder: NSCoder) {
+        // Font to use in our graph title
+        self.titleFont = UIFont.systemFont(ofSize: 12 / (UIScreen.main.scale - 1));
+        
+        // Font to use in our graph tick labels
+        self.tickLabelFont = UIFont.systemFont(ofSize: 8 / (UIScreen.main.scale - 1));
+        
         super.init(coder: coder)
         self.backgroundColor = UIColor.clear
     }
diff --git a/Engine/Extras/VirtualProduction/LiveLinkVCAM/vcam/StartViewController+StreamingConnectionDelegate.swift b/Engine/Extras/VirtualProduction/LiveLinkVCAM/vcam/StartViewController+StreamingConnectionDelegate.swift
--- a/Engine/Extras/VirtualProduction/LiveLinkVCAM/vcam/StartViewController+StreamingConnectionDelegate.swift
+++ b/Engine/Extras/VirtualProduction/LiveLinkVCAM/vcam/StartViewController+StreamingConnectionDelegate.swift
@@ -13,31 +13,30 @@ extension StartViewController : StreamingConnectionDelegate {
     func streamingConnectionDidConnect(_ connection: StreamingConnection) {
         Log.info("StreamingConnection \(connection.name) did connect to \(connection.destination)")
         
-        DispatchQueue.main.async {
-            if self.connectingView.isHidden == true {
-                self.performSegue(withIdentifier: "showVideoView", sender: self)
+        DispatchQueue.main.async { [weak self] in
+            if let validSelf = self, validSelf.connectingView.isHidden {
+                validSelf.performSegue(withIdentifier: "showVideoView", sender: self)
             } else {
-                self.hideConnectingView {
-                    self.performSegue(withIdentifier: "showVideoView", sender: self)
+                self?.hideConnectingView { [weak self] in
+                    self?.performSegue(withIdentifier: "showVideoView", sender: self)
                 }
             }
-            
         }
     }
 
     func streamingConnection(_ connection: StreamingConnection, didDisconnectWithError err: Error?) {
 
-        DispatchQueue.main.async {
+        DispatchQueue.main.async { [weak self] in
             connection.disconnect()
-            self.hideConnectingAlertView {
+            self?.hideConnectingAlertView { [weak self] in
                 if let e = err {
                     Log.info("StreamingConnection \(connection.name) disconnected with error : \(e.localizedDescription)")
 
                     let errorAlert = UIAlertController(title: Localized.titleError(), message: "\(Localized.messageCouldntConnect()) : \(e.localizedDescription)", preferredStyle: .alert)
                     errorAlert.addAction(UIAlertAction(title: Localized.buttonOK(), style: .default, handler: { _ in
-                        self.hideConnectingView() {}
+                        self?.hideConnectingView() {}
                     }))
-                    self.present(errorAlert, animated:true)
+                    self?.present(errorAlert, animated:true)
                     
                 } else {
                     Log.info("StreamingConnection \(connection.name) disconnected.")
@@ -47,17 +46,17 @@ extension StartViewController : StreamingConnectionDelegate {
     }
     func streamingConnection(_ connection: StreamingConnection, exitWithError err: Error?) {
 
-        DispatchQueue.main.async {
+        DispatchQueue.main.async { [weak self] in
             connection.disconnect()
-            self.hideConnectingAlertView {
+            self?.hideConnectingAlertView { [weak self] in
                 if let e = err {
                     Log.info("StreamingConnection \(connection.name) disconnected with error : \(e.localizedDescription)")
 
                     let errorAlert = UIAlertController(title: Localized.titleError(), message: "\(Localized.messageCouldntConnect()) : \(e.localizedDescription)", preferredStyle: .alert)
-                    errorAlert.addAction(UIAlertAction(title: Localized.buttonOK(), style: .default, handler: { _ in
-                        self.hideConnectingView() {}
+                    errorAlert.addAction(UIAlertAction(title: Localized.buttonOK(), style: .default, handler: { [weak self] _ in
+                        self?.hideConnectingView() {}
                     }))
-                    self.present(errorAlert, animated:true)
+                    self?.present(errorAlert, animated:true)
                     
                 } else {
                     Log.info("StreamingConnection \(connection.name) disconnected.")
@@ -72,20 +71,21 @@ extension StartViewController : StreamingConnectionDelegate {
     func streamingConnection(_ connection: StreamingConnection, requestStreamerSelectionWithStreamers streamers: Array<String>, handler: @escaping (String) -> Void) {
         self.pickerData = streamers;
         self.selectedStreamer = streamers[0];
-        DispatchQueue.main.async {
-            self.hideConnectingView() {
+        DispatchQueue.main.async { [weak self] in
+            self?.hideConnectingView() { [weak self] in
                 let alert = UIAlertController(title: Localized.titleSelectStream(), message: "\n\n\n\n\n\n", preferredStyle: .alert)
-                
                 let picker = UIPickerView(frame: CGRect(x: 5, y: 20, width: 250, height: 140))
-                picker.dataSource = self;
-                picker.delegate = self;
+                picker.dataSource = self
+                picker.delegate = self
                 alert.view.addSubview(picker)
 
-                alert.addAction(UIAlertAction(title: Localized.buttonOK(), style: .default) {_ in
-                    handler(self.selectedStreamer)
+                alert.addAction(UIAlertAction(title: Localized.buttonOK(), style: .default) { [weak self] _ in
+                    if let chosenStreamer = self?.selectedStreamer {
+                        handler(chosenStreamer)
+                    }
                 })
                 
-                self.present(alert, animated:true)
+                self?.present(alert, animated:true)
             }
         }
     }
diff --git a/Engine/Extras/VirtualProduction/LiveLinkVCAM/vcam/StartViewController.swift b/Engine/Extras/VirtualProduction/LiveLinkVCAM/vcam/StartViewController.swift
--- a/Engine/Extras/VirtualProduction/LiveLinkVCAM/vcam/StartViewController.swift
+++ b/Engine/Extras/VirtualProduction/LiveLinkVCAM/vcam/StartViewController.swift
@@ -28,14 +28,13 @@ class StartViewController : BaseViewController {
     private var tapGesture : UITapGestureRecognizer!
     
     private var streamingConnection : StreamingConnection?
+    private var gameController : GCController?
     private var _liveLinkTimer : Timer?
 
     
     @objc dynamic let appSettings = AppSettings.shared
     private var observers = [NSKeyValueObservation]()
     
-    private var gameController : GCController?
-    
     var pickerData: [String] = [String]()
     var selectedStreamer: String = "";
 
@@ -47,36 +46,27 @@ class StartViewController : BaseViewController {
           return .lightContent
     }
     
-    override func viewDidLoad() {
-        
-        super.viewDidLoad();
-        
-        self.restartView.isHidden = true
-        
-        if let infoDict = Bundle.main.infoDictionary {
-            self.versionLabel.text = String(format: "v%@ (%@)", infoDict["CFBundleShortVersionString"] as! String, infoDict["CFBundleVersion"] as! String)
-        } else {
-            self.versionLabel.text = "";
-        }
-        
-        NetUtility.triggerLocalNetworkPrivacyAlert()
-        
-        self.ipAddress.text = AppSettings.shared.lastConnectionAddress
-        self.ipAddress.inputAssistantItem.leadingBarButtonGroups.removeAll()
-        self.ipAddress.inputAssistantItem.trailingBarButtonGroups.removeAll()
-        textFieldChanged(self.ipAddress)
-        
-        self.rebuildRecentAddressesBarButtons()
-
-        self.tapGesture = UITapGestureRecognizer(target: self, action: #selector(handleTap))
-        self.tapGesture.cancelsTouchesInView = false
-        self.tapGesture.delegate = self
-        self.view.addGestureRecognizer(tapGesture)
-        
+    // Constructor (before view is even loaded)
+    required init?(coder: NSCoder) {
+        super.init(coder: coder)
+        setupObservers()
+    }
+    
+    deinit {
+        Log.info("StartViewController destructed")
+    }
+    
+    func setupObservers() {
+        // Observers for keyboard show/hide
         NotificationCenter.default.addObserver(self, selector: #selector(keyboardWillShow), name: UIResponder.keyboardWillShowNotification, object: nil)
         NotificationCenter.default.addObserver(self, selector: #selector(keyboardWillHide), name: UIResponder.keyboardWillHideNotification, object: nil)
         
-        observers.append(observe(\.appSettings.timecodeSource, options: [.initial,.new,.old], changeHandler: { object, change in
+        // Observers for game controller connect/disconnect
+        NotificationCenter.default.addObserver(self, selector: #selector(gameControllerDidConnectNotification), name: .GCControllerDidConnect, object: nil)
+        NotificationCenter.default.addObserver(self, selector: #selector(gameControllerDidDisconnectNotification), name: .GCControllerDidDisconnect, object: nil)
+        
+        // Observers for app settings
+        observers.append(observe(\.appSettings.timecodeSource, options: [.initial,.new,.old], changeHandler: { [weak self] object, change in
             
             if let oldValue = TimecodeSource(rawValue:change.oldValue ?? 0) {
                 switch oldValue {
@@ -89,7 +79,7 @@ class StartViewController : BaseViewController {
                 }
             }
             
-            switch self.appSettings.timecodeSourceEnum() {
+            switch self?.appSettings.timecodeSourceEnum() {
             case .ntp:
                 let pool = AppSettings.shared.ntpPool.isEmpty ? "time.apple.com" : AppSettings.shared.ntpPool
                 Log.info("Started NTP : \(pool)")
@@ -109,29 +99,56 @@ class StartViewController : BaseViewController {
         }))
         
         // any change to the subject name will remove & re-add the camera subject.
-        observers.append(observe(\.appSettings.liveLinkSubjectName, options: [.old,.new], changeHandler: { object, change in
-            if let sc = self.streamingConnection {
-                sc.subjectName = self.appSettings.liveLinkSubjectName
+        observers.append(observe(\.appSettings.liveLinkSubjectName, options: [.old,.new], changeHandler: { [weak self] object, change in
+            if let sc = self?.streamingConnection {
+                sc.subjectName = self?.appSettings.liveLinkSubjectName
             }
         }))
 
         // initial & value changes for the connection type instantiates a new StreamingConnection object
-        observers.append(observe(\.appSettings.connectionType, options: [.initial, .old,.new], changeHandler: { object, change in
+        observers.append(observe(\.appSettings.connectionType, options: [.initial, .old,.new], changeHandler: { [weak self] object, change in
             
-            self.streamingConnection?.shutdown()
-            self.streamingConnection = nil
-
-            let connectionType = self.appSettings.connectionType
-            if let connectionClass = Bundle.main.classNamed("VCAM.\(connectionType)StreamingConnection") as? StreamingConnection.Type {
-                self.streamingConnection = connectionClass.init(subjectName: self.appSettings.liveLinkSubjectName)
-                self.streamingConnection?.delegate = self
+            if let validSelf = self {
+                validSelf.streamingConnection?.shutdown()
+                validSelf.streamingConnection = nil
+
+                let connectionType = validSelf.appSettings.connectionType
+                if let connectionClass = Bundle.main.classNamed("VCAM.\(connectionType)StreamingConnection") as? StreamingConnection.Type {
+                    validSelf.streamingConnection = connectionClass.init(subjectName: validSelf.appSettings.liveLinkSubjectName)
+                    validSelf.streamingConnection?.delegate = validSelf
+                }
             }
-
         }))
+    }
+    
+    override func viewDidLoad() {
+        
+        super.viewDidLoad();
+        
+        self.restartView.isHidden = true
+        
+        if let infoDict = Bundle.main.infoDictionary {
+            self.versionLabel.text = String(format: "v%@ (%@)", infoDict["CFBundleShortVersionString"] as! String, infoDict["CFBundleVersion"] as! String)
+        } else {
+            self.versionLabel.text = "";
+        }
+        
+        NetUtility.triggerLocalNetworkPrivacyAlert()
+        
+        self.ipAddress.text = AppSettings.shared.lastConnectionAddress
+        self.ipAddress.inputAssistantItem.leadingBarButtonGroups.removeAll()
+        self.ipAddress.inputAssistantItem.trailingBarButtonGroups.removeAll()
+        textFieldChanged(self.ipAddress)
+        
+        self.rebuildRecentAddressesBarButtons()
 
-        NotificationCenter.default.addObserver(self, selector: #selector(gameControllerDidConnectNotification), name: .GCControllerDidConnect, object: nil)
-        NotificationCenter.default.addObserver(self, selector: #selector(gameControllerDidDisconnectNotification), name: .GCControllerDidDisconnect, object: nil)
+        // Add gesture recognizer
+        self.tapGesture = UITapGestureRecognizer(target: self, action: #selector(handleTap))
+        self.tapGesture.cancelsTouchesInView = false
+        self.tapGesture.delegate = self
+        self.view.addGestureRecognizer(tapGesture)
         
+        // Attempt to get an attached game controller
         self.gameController = GCController.controllers().first
         if let gc = self.gameController {
             if gc.isAttachedToDevice {
@@ -143,25 +160,28 @@ class StartViewController : BaseViewController {
 
     
     override func viewWillAppear(_ animated : Bool) {
-        
         self.connectingView.isHidden = true
         self.headerView.start()
 
         self.streamingConnection?.delegate = self
         
-        _liveLinkTimer = Timer.scheduledTimer(withTimeInterval: 1.0/10.0, repeats: true, block: { timer in
-            self.streamingConnection?.sendTransform(simd_float4x4(), atTime: Timecode.create().toTimeInterval())
+        _liveLinkTimer = Timer.scheduledTimer(withTimeInterval: 1.0/10.0, repeats: true, block: { [weak self] timer in
+            self?.streamingConnection?.sendTransform(simd_float4x4(), atTime: Timecode.create().toTimeInterval())
         })
     }
     
     override func viewDidAppear(_ animated: Bool) {
         super.viewDidAppear(animated)
+        self.headerView.start()
         self.streamingConnection?.disconnect()
     }
     
     override func viewDidDisappear(_ animated: Bool) {
         super.viewDidDisappear(animated)
         self.headerView.stop()
+        
+        // Remove gesture recognizer
+        self.view.removeGestureRecognizer(self.tapGesture)
     }
     
     override func prepare(for segue: UIStoryboardSegue, sender: Any?) {
@@ -181,9 +201,11 @@ class StartViewController : BaseViewController {
                 _liveLinkTimer?.invalidate()
                 _liveLinkTimer = nil
 
-                self.streamingConnection?.delegate = vc
+                // This a little awkward but what happens here is the streamingConnection and gameController gets
+                // are passed (weakly) from StartViewController to VideoViewController as it needs them
+                // They are not nil'd in this view because this view still exists because it is the ancestor of the segue
                 vc.streamingConnection = self.streamingConnection
-                
+                vc.streamingConnection?.delegate = vc
                 vc.gameController = self.gameController
             }
         }
@@ -206,8 +228,8 @@ class StartViewController : BaseViewController {
             self.entryViewYConstraint.constant = 0
         }
 
-        UIView.animate(withDuration: 0.2) {
-            self.view.layoutIfNeeded()
+        UIView.animate(withDuration: 0.2) { [weak self] in
+            self?.view.layoutIfNeeded()
         }
     }
 
@@ -215,8 +237,8 @@ class StartViewController : BaseViewController {
         if self.entryViewYConstraint.constant != 0 {
             self.entryViewYConstraint.constant = 0
 
-            UIView.animate(withDuration: 0.2) {
-                self.view.layoutIfNeeded()
+            UIView.animate(withDuration: 0.2) { [weak self] in
+                self?.view.layoutIfNeeded()
             }
         }
     }
@@ -271,7 +293,6 @@ class StartViewController : BaseViewController {
     }
     
     @objc func handleTap(_ recognizer: UITapGestureRecognizer) {
-        
         self.view.endEditing(true)
     }
     
@@ -288,13 +309,13 @@ class StartViewController : BaseViewController {
             // show the connection view
             self.connectingView.isHidden = false
             self.connectingView.alpha = 0.0
-            UIView.animate(withDuration: 0.2) {
-                self.connectingView.alpha = 1.0
+            UIView.animate(withDuration: 0.2) { [weak self] in
+                self?.connectingView.alpha = 1.0
             }
             
-            showConnectingAlertView(mode: .connecting) {
-                self.hideConnectingView() { }
-            }
+            showConnectingAlertView(mode: .connecting, { [weak self] in
+                self?.hideConnectingView() {}
+            })
 
             do {
                 self.streamingConnection?.destination = self.ipAddress.text!.trimmed()
@@ -315,18 +336,18 @@ class StartViewController : BaseViewController {
         hideConnectingAlertView() {
 
             let errorAlert = UIAlertController(title: Localized.titleError(), message: "\(Localized.messageCouldntConnect()) : \(message)", preferredStyle: .alert)
-            errorAlert.addAction(UIAlertAction(title: Localized.buttonOK(), style: .default, handler: { _ in
-                self.hideConnectingView { }
+            errorAlert.addAction(UIAlertAction(title: Localized.buttonOK(), style: .default, handler: { [weak self] _ in
+                self?.hideConnectingView { }
             }))
             self.present(errorAlert, animated: true)
         }
     }
     
     func hideConnectingView( _ completion : @escaping () -> Void) {
-        UIView.animate(withDuration: 0.2, animations: {
-            self.connectingView.alpha = 0.0
-        }, completion: { b in
-            self.connectingView.isHidden = true
+        UIView.animate(withDuration: 0.2, animations: { [weak self] in
+            self?.connectingView.alpha = 0.0
+        }, completion: { [weak self] b in
+            self?.connectingView.isHidden = true
         })
         hideConnectingAlertView(completion)
     }
diff --git a/Engine/Extras/VirtualProduction/LiveLinkVCAM/vcam/StreamingConnection.swift b/Engine/Extras/VirtualProduction/LiveLinkVCAM/vcam/StreamingConnection.swift
--- a/Engine/Extras/VirtualProduction/LiveLinkVCAM/vcam/StreamingConnection.swift
+++ b/Engine/Extras/VirtualProduction/LiveLinkVCAM/vcam/StreamingConnection.swift
@@ -68,10 +68,10 @@ class StreamingConnectionStats {
 class StreamingConnection : NSObject {
 
     weak var delegate: StreamingConnectionDelegate?
+    weak var renderView : UIView?
+    
     var subjectName : String!
     
-    var renderView : UIView?
-
     var name: String {
         get {
             assertionFailure("not implemented")
diff --git a/Engine/Extras/VirtualProduction/LiveLinkVCAM/vcam/VideoViewController+StreamingConnectionDelegate.swift b/Engine/Extras/VirtualProduction/LiveLinkVCAM/vcam/VideoViewController+StreamingConnectionDelegate.swift
--- a/Engine/Extras/VirtualProduction/LiveLinkVCAM/vcam/VideoViewController+StreamingConnectionDelegate.swift
+++ b/Engine/Extras/VirtualProduction/LiveLinkVCAM/vcam/VideoViewController+StreamingConnectionDelegate.swift
@@ -13,8 +13,8 @@ extension VideoViewController : StreamingConnectionDelegate {
 
     func streamingConnectionDidConnect(_ connection: StreamingConnection) {
         Log.info("StreamingConnection \(connection.name) did connect to \(connection.destination)")
-        DispatchQueue.main.async {
-            self.showReconnecting(false, animated: true)
+        DispatchQueue.main.async { [weak self] in
+            self?.showReconnecting(false, animated: true)
         }
     }
     
@@ -25,66 +25,64 @@ extension VideoViewController : StreamingConnectionDelegate {
             Log.info("StreamingConnection \(connection.name) disconnected.")
         }
 
-        DispatchQueue.main.async {
-
-            if self.dismissOnDisconnect {
-                self.forceDisconnectAndDismiss()
-            } else {
-                self.reconnect()
-            }
+        // Got connection error, so exit back to streaming screen.
+        DispatchQueue.main.async { [weak self] in
+            self?.reconnect()
         }
     }
     
     func streamingConnection(_ connection: StreamingConnection, exitWithError err: Error?) {
-        DispatchQueue.main.async {
-            self.hideConnectingAlertView() {
+        DispatchQueue.main.async { [weak self] in
+            self?.hideConnectingAlertView() {
                 let alert = UIAlertController(title: nil, message: err?.localizedDescription, preferredStyle: .alert)
-                alert.addAction(UIAlertAction(title: Localized.buttonOK(), style: .default) { _ in
-                    self.disconnect()
+                alert.addAction(UIAlertAction(title: Localized.buttonOK(), style: .default) { [weak self] _ in
+                    self?.disconnect()
                 })
                 
-                self.present(alert, animated: true)
+                self?.present(alert, animated: true)
             }
         }
     }
     
     func streamingConnection(_ connection: StreamingConnection, requestsTextEditWithContents contents: String, handler: @escaping (Bool, String?) -> Void) {
         
-        DispatchQueue.main.async {
+        DispatchQueue.main.async { [weak self] in
             let alert = UIAlertController(title: nil, message: contents.count == 0 ?
                                           NSLocalizedString("message-entervalue", value: "Enter the value", comment: "Entering value in a textbox") :
                                           NSLocalizedString("message-updatevalue", value: "Update the value", comment: "Updating value in a textbox"),
                                           preferredStyle: .alert)
             alert.addTextField()
             alert.textFields![0].text = contents
-            alert.addAction(UIAlertAction(title: Localized.buttonOK(), style: .default) {_ in
-                handler(true, alert.textFields![0].text!)
+            alert.addAction(UIAlertAction(title: Localized.buttonOK(), style: .default) {[weak alert] _ in
+                handler(true, alert?.textFields![0].text!)
             })
             alert.addAction(UIAlertAction(title: Localized.buttonCancel(), style: .cancel) {_ in
                 handler(false, nil)
             })
 
-            self.present(alert, animated: true)
+            self?.present(alert, animated: true)
         }
     }
     
     func streamingConnection(_ connection: StreamingConnection, requestStreamerSelectionWithStreamers streamers: Array<String>, handler: @escaping (String) -> Void) {
         self.pickerData = streamers;
         self.selectedStreamer = streamers[0];
-        DispatchQueue.main.async {
-            self.hideConnectingAlertView() {
+        DispatchQueue.main.async { [weak self] in
+            self?.hideConnectingAlertView() { [weak self] in
                 let alert = UIAlertController(title: Localized.titleSelectStream(), message: "\n\n\n\n\n\n", preferredStyle: .alert)
                 
                 let picker = UIPickerView(frame: CGRect(x: 5, y: 20, width: 250, height: 140))
-                picker.dataSource = self;
-                picker.delegate = self;
+                picker.dataSource = self
+                picker.delegate = self
                 alert.view.addSubview(picker)
 
-                alert.addAction(UIAlertAction(title: Localized.buttonOK(), style: .default) {_ in
-                    handler(self.selectedStreamer)
+                alert.addAction(UIAlertAction(title: Localized.buttonOK(), style: .default) {[weak self] _ in
+                    if let selectedStreamer = self?.selectedStreamer {
+                        handler(selectedStreamer)
+                    }
                 })
                 
-                self.present(alert, animated:true)
+                self?.present(alert, animated:true)
             }
         }
     }
diff --git a/Engine/Extras/VirtualProduction/LiveLinkVCAM/vcam/VideoViewController.swift b/Engine/Extras/VirtualProduction/LiveLinkVCAM/vcam/VideoViewController.swift
--- a/Engine/Extras/VirtualProduction/LiveLinkVCAM/vcam/VideoViewController.swift
+++ b/Engine/Extras/VirtualProduction/LiveLinkVCAM/vcam/VideoViewController.swift
@@ -17,36 +17,37 @@ class VideoViewController : BaseViewController {
     var demoMode : Bool {
         liveLink == nil
     }
+    
+    weak var liveLink : LiveLinkProvider?
+    
+    // Required public for reconnection picker
+    var pickerData: [String] = [String]()
+    var selectedStreamer: String = ""
 
     private var displayLink: CADisplayLink?
-    
     private var refreshRateHint : CADisplayLink?
     private var lastTimestamp: CFTimeInterval = 0.0
 
     private var arSession : ARSession?
+    private var arCoachingView : ARCoachingOverlayView?
     
     @IBOutlet weak var renderView : UIView!
     
     @IBOutlet weak var headerView : HeaderView!
     @IBOutlet weak var headerViewTopConstraint : NSLayoutConstraint!
-    var headerViewY : CGFloat = 0
-    var headerViewHeight : CGFloat = 0
-    var headerViewTopConstraintStartValue : CGFloat = 0
-    var headerPanGestureRecognizer : UIPanGestureRecognizer!
-    var headerPullDownGestureRecognizer : UIScreenEdgePanGestureRecognizer!
+    private var headerViewY : CGFloat = 0
+    private var headerViewHeight : CGFloat = 0
+    private var headerViewTopConstraintStartValue : CGFloat = 0
+    private var headerPanGestureRecognizer : UIPanGestureRecognizer!
+    private var headerPullDownGestureRecognizer : UIScreenEdgePanGestureRecognizer!
     
-    var streamingConnection : StreamingConnection?
-    var pickerData: [String] = [String]()
-    var selectedStreamer: String = ""
-
-    weak var liveLink : LiveLinkProvider?
-    var dismissOnDisconnect = false
-    var forceDisconnectTimer : Timer?
+    private var statsTimer : Timer?
+    private var showStats : Bool = false
     
-    var statsTimer : Timer?
-    var showStats : Bool = false
+    private var gameControllerSnapshot : GCController?
     
-    var gameControllerSnapshot : GCController?
+    // streamingConnection and gameController are passed from StartViewController
+    weak var streamingConnection : StreamingConnection?
     weak var gameController : GCController? {
         didSet {
             if let gc = gameController {
@@ -95,51 +96,107 @@ class VideoViewController : BaseViewController {
     @objc func applicationDidEnterBackground(notification: NSNotification) {
         disconnect()
     }
+    @objc func controllerDidConnect(notification: NSNotification) {
+        guard let gc = notification.object as? GCController else { return }
+        self.controllerConnected(gamepad: gc)
+    }
+    @objc func controllerDidDisconnect(notification: NSNotification) {
+        guard let gc = notification.object as? GCController else { return }
+        self.controllerDisconnected(gamepad: gc)
+    }
     
-    override func viewDidLoad() {
-        super.viewDidLoad()
+    // Constructor (before view is even loaded)
+    required init?(coder: NSCoder) {
+        super.init(coder: coder)
+        setupObservers()
+        setupARSession()
+    }
+    
+    // Destructor (called when view controller is destroyed)
+    deinit {
+        Log.info("VideoViewController destructed.")
+        self.headerView = nil
+    }
+    
+    func setupObservers(){
+        // According to: https://stackoverflow.com/a/40339926
+        // There is no need to remove observers because they are captured weakly and automatically removed (as long as they are not using closure blocks)
         
         NotificationCenter.default.addObserver(self, selector: #selector(applicationDidBecomeActive), name: UIApplication.willEnterForegroundNotification, object: nil)
         NotificationCenter.default.addObserver(self, selector: #selector(applicationDidEnterBackground), name: UIApplication.didEnterBackgroundNotification, object: nil)
         
-        self.headerView.start()
-        headerViewHeight = self.headerView.frame.height
-
-        headerPanGestureRecognizer = UIPanGestureRecognizer(target: self, action: #selector(handleHeaderPanGesture))
-        self.headerView.addGestureRecognizer(headerPanGestureRecognizer)
+        // Add notifications for when new controllers connect or disconnect
+        NotificationCenter.default.addObserver(self, selector: #selector(controllerDidConnect), name: .GCControllerDidConnect, object: nil)
+        NotificationCenter.default.addObserver(self, selector: #selector(controllerDidDisconnect), name: .GCControllerDidDisconnect, object: nil)
+    }
+    
+    func setupGestureRecognizers(){
+        self.headerPanGestureRecognizer = UIPanGestureRecognizer(target: self, action: #selector(handleHeaderPanGesture))
+        self.headerView.addGestureRecognizer(self.headerPanGestureRecognizer)
 
-        headerPullDownGestureRecognizer = UIScreenEdgePanGestureRecognizer(target: self, action: #selector(handleHeaderPanGesture))
-        headerPullDownGestureRecognizer.edges = [ .top ]
-        self.view.addGestureRecognizer(headerPullDownGestureRecognizer)
-        
-        headerPanGestureRecognizer.require(toFail: headerPullDownGestureRecognizer)
-        
-        // The AR video feed & CG objects only exist in demo mode. In normal operation, they are replaced by a live feed of the Unreal Engine rendering, with the virtual camera's position determined by the device & ARKit.
-        //self.demoModeBlurView.isHidden = !self.demoMode
-        showReconnecting(false, animated: false)
+        self.headerPullDownGestureRecognizer = UIScreenEdgePanGestureRecognizer(target: self, action: #selector(handleHeaderPanGesture))
+        self.headerPullDownGestureRecognizer.edges = [ .top ]
+        self.view.addGestureRecognizer(self.headerPullDownGestureRecognizer)
         
+        self.headerPanGestureRecognizer.require(toFail: headerPullDownGestureRecognizer)
+    }
+    
+    func removeGestureRecognizers(){
+        self.headerView.removeGestureRecognizer(self.headerPanGestureRecognizer)
+        self.view.removeGestureRecognizer(self.headerPullDownGestureRecognizer)
+    }
+    
+    func setupARSession(){
+        arSession = ARSession()
+        arSession?.delegate = self
+    }
+    
+    func startARSession(){
         // Create a ARKit tracking config for purely tracking device transform (we don't care about the computer vision features)
         let config = ARPositionalTrackingConfiguration()
         config.worldAlignment = .gravity
         config.planeDetection = []
         config.isLightEstimationEnabled = false
         config.providesAudioData = false
+        if self.arSession == nil {
+            setupARSession()
+        }
+        self.arSession?.run(config)
+    }
+    
+    override func viewDidLoad() {
+        super.viewDidLoad()
         
-        arSession = ARSession()
-        arSession?.delegate = self
-        arSession?.run(config)
+        self.headerView.start()
+        headerViewHeight = self.headerView.frame.height
+        
+        // The AR video feed & CG objects only exist in demo mode. 
+        // In normal operation, they are replaced by a live feed of the Unreal Engine rendering, with the virtual camera's position determined by the device & ARKit.
+        //self.demoModeBlurView.isHidden = !self.demoMode
+        showReconnecting(false, animated: false)
+        
+        // setup gesture recognition has to start here because it references view that need to have loaded
+        self.setupGestureRecognizers()
         
+        // Start running AR session
+        self.startARSession()
+        
+        // Attach an empty view where the relevant stream connection can setup its own view once this is set
         self.streamingConnection?.renderView = self.renderView
         
+        if let existingARCoachingView = self.arCoachingView {
+            existingARCoachingView.removeFromSuperview()
+        }
         
+        // Insert coaching overlay into subview
         let coachingOverlayView = ARCoachingOverlayView()
-        
         self.view.insertSubview(coachingOverlayView, belowSubview: self.reconnectingBlurView)
         coachingOverlayView.layout(.left, .right, .top, .bottom, to: self.arView)
         coachingOverlayView.goal = .tracking
         coachingOverlayView.session = self.arSession
         coachingOverlayView.activatesAutomatically = true
         coachingOverlayView.delegate = self
+        self.arCoachingView = coachingOverlayView
 
         statsTimer = Timer.scheduledTimer(withTimeInterval: 1.0, repeats: true, block: { t in
             self.updateStreamingStats()
@@ -148,25 +205,14 @@ class VideoViewController : BaseViewController {
     
     override func viewDidAppear(_ animated: Bool) {
         super.viewDidAppear(animated)
+        
+        // Ensure we do not go into sleep mode while this view is active
         UIApplication.shared.isIdleTimerDisabled = true
         
         // Check for already connected controllers
         for gc in GCController.controllers() {
             self.controllerConnected(gamepad: gc)
         }
-
-        let notificationCenter = NotificationCenter.default;
-        
-        // Add notifications for when new controllers connect or disconnect
-        notificationCenter.addObserver(forName: .GCControllerDidConnect, object: nil, queue: .main) { (note) in
-            guard let gc = note.object as? GCController else { return }
-            self.controllerConnected(gamepad: gc)
-        }
-        
-        notificationCenter.addObserver(forName: .GCControllerDidDisconnect, object: nil, queue: .main) { (note) in
-            guard let gc = note.object as? GCController else { return }
-            self.controllerDisconnected(gamepad: gc)
-        }
         
         // Todo: Could dynamically adjust this
         self.setRefreshRateFps(fps: 60)
@@ -174,14 +220,25 @@ class VideoViewController : BaseViewController {
     
     override func viewWillDisappear(_ animated: Bool) {
         super.viewWillDisappear(animated)
-        UIApplication.shared.isIdleTimerDisabled = false
         
-        let notificationCenter = NotificationCenter.default
-        notificationCenter.removeObserver(self, name: .GCControllerDidConnect, object: nil)
-        notificationCenter.removeObserver(self, name: .GCControllerDidDisconnect, object: nil)
+        // Now this view is going away we can let the device go to sleep if there is not touch input
+        UIApplication.shared.isIdleTimerDisabled = false
         
         // If view is going away we will remove out refresh rate hint so battery usage can return to normal values
         self.resetRefreshRateFps()
+        
+        // Clear stats timer
+        self.statsTimer?.invalidate()
+        
+        // pause ar session when view goes away
+        self.arSession?.pause()
+        
+        // remove the gesture recognizers as this view will not longer be shown
+        self.removeGestureRecognizers()
+        
+        // disconnect from streaming if we haven't already
+        self.disconnect()
+        
     }
     
     @objc func refreshRateCallback(_ displayLink: CADisplayLink) {
@@ -203,7 +260,7 @@ class VideoViewController : BaseViewController {
         
         for subview in self.renderView.subviews {
             if let webrtcView = subview as? WebRTCView {
-                webrtcView.videoView.setPreferredFramerate(fps: fps)
+                webrtcView.videoView?.setPreferredFramerate(fps: fps)
             }
         }
     }
@@ -235,42 +292,29 @@ class VideoViewController : BaseViewController {
     }
 
     func exit() {
-        
-        // if still connected, we need to send a disconnect/goodbye, then dismiss the view controller when
-        // the goodbye msg is sent.
-        // otherwise, we can call disconnect (which tears down some objects) and dismiss this VC immediately.
-        if self.streamingConnection?.isConnected ?? false {
-            dismissOnDisconnect = true
-            disconnect()
-        } else {
-            self.forceDisconnectAndDismiss()
-        }
+        disconnect()
     }
     
     func reconnect() {
-        
         showReconnecting(true, animated: true)
         self.streamingConnection?.reconnect()
     }
     
     func disconnect() {
-        self.forceDisconnectTimer = Timer.scheduledTimer(withTimeInterval: 0.5, repeats: false, block: { timer in
-            Log.info("Forcing disconnect due to timeout.")
-            self.forceDisconnectAndDismiss()
-        })
         // Tell UE that the controllers connected to this device are no longer used
         for gc in GCController.controllers() {
             self.controllerDisconnected(gamepad: gc)
         }
         
         self.streamingConnection?.disconnect()
-    }
-    
-    func forceDisconnectAndDismiss() {
-        self.forceDisconnectTimer?.invalidate()
-        self.forceDisconnectTimer = nil
+        
+        // Clear the delegate on the streaming connection (which is a reference to this view controller)
         self.streamingConnection?.delegate = nil
+        
+        // Clear the streaming connection itself
         self.streamingConnection = nil
+        
+        // Dismiss this UIViewController and return to the presenting view that segued to here
         self.presentingViewController?.dismiss(animated: true, completion: nil)
     }
     
@@ -349,19 +393,17 @@ class VideoViewController : BaseViewController {
             
             // we also need a display link here to properly communicate the new headerY from the presentation layer to
             // the metal renderer.
-            self.displayLink = CADisplayLink(
-              target: self, selector: #selector(displayLinkDidFire)
-            )
+            self.displayLink = CADisplayLink(target: self, selector: #selector(displayLinkDidFire))
             self.displayLink!.add(to: .main, forMode: .common)
 
-            UIView.animate(withDuration: 0.2) {
+            UIView.animate(withDuration: 0.2) { [weak self] in
                 // animate to the new position
-                self.headerViewTopConstraint.constant = newTopY
-                self.view.layoutIfNeeded()
-            } completion: { _ in
+                self?.headerViewTopConstraint.constant = newTopY
+                self?.view.layoutIfNeeded()
+            } completion: { [weak self] _ in
                 // all done, we kill the display link
-                self.displayLink?.invalidate()
-                self.displayLink = nil
+                self?.displayLink?.invalidate()
+                self?.displayLink = nil
             }
         }
     }
@@ -393,7 +435,6 @@ class VideoViewController : BaseViewController {
         ]
         
         for input in inputs {
-            
             if (input.value.newValue != 0.0) || ((input.value.oldValue ?? 0.0) != input.value.newValue) {
                 sc.sendControllerAnalog(input.key, controllerIndex: UInt8(controllerIndex), value: input.value.newValue)
             }
@@ -498,15 +539,14 @@ extension VideoViewController : HeaderViewDelegate {
     func headerViewExitButtonTapped(_ headerView : HeaderView) {
      
         let disconnectAlert = UIAlertController(title: nil, message: NSLocalizedString("Disconnect from the remote session?", comment: "Prompt disconnect from a remote session."), preferredStyle: .alert)
-        disconnectAlert.addAction(UIAlertAction(title: NSLocalizedString("Disconnect", comment: "Button to disconnect from a UE instance"), style: .destructive, handler: { _ in
-            self.exit()
+        disconnectAlert.addAction(UIAlertAction(title: NSLocalizedString("Disconnect", comment: "Button to disconnect from a UE instance"), style: .destructive, handler: { [weak self] _ in
+            self?.exit()
         }))
         disconnectAlert.addAction(UIAlertAction(title: Localized.buttonCancel(), style: .cancel))
         self.present(disconnectAlert, animated:true)
     }
 
     func headerViewLogButtonTapped(_ headerView : HeaderView) {
-     
         performSegue(withIdentifier: "showLog", sender: headerView)
     }
     
diff --git a/Engine/Extras/VirtualProduction/LiveLinkVCAM/vcam/WebRTC/Extensions/TouchControls.swift b/Engine/Extras/VirtualProduction/LiveLinkVCAM/vcam/WebRTC/Extensions/TouchControls.swift
--- a/Engine/Extras/VirtualProduction/LiveLinkVCAM/vcam/WebRTC/Extensions/TouchControls.swift
+++ b/Engine/Extras/VirtualProduction/LiveLinkVCAM/vcam/WebRTC/Extensions/TouchControls.swift
@@ -35,8 +35,8 @@ struct TouchPoint {
 
 final class TouchControls : TouchDelegate {
     
-    let webRTCClient : WebRTCClient
-    let touchView : UIView
+    weak var webRTCClient : WebRTCClient?
+    weak var touchView : UIView?
     var videoAspectRatio : CGFloat
     
     // We need a way to give each touch a unique finger id that is persistent throughout
@@ -65,6 +65,10 @@ final class TouchControls : TouchDelegate {
         self.videoAspectRatio = CGFloat.zero
     }
     
+    deinit {
+        Log.info("TouchControls destructed.")
+    }
+    
     func rememberTouch(_ touch : UITouch) {
         
         // use the smallest fingerId available
@@ -126,35 +130,40 @@ final class TouchControls : TouchDelegate {
     
     // Unsigned XY positions are normalized to the ratio (0.0..1.0) along a viewport axis and then quantized into an uint16 (0..65536).
     func normalizeAndQuantize(_ touch: UITouch) -> TouchPoint {
-        let uiAspectRatio : CGFloat = touchView.bounds.width / touchView.bounds.height
-        let touchLocation : CGPoint = touch.location(in: touchView)
-        let viewBounds : CGRect = touchView.bounds
-        
-        // normalise x,y to the UI element bounds
-        var normalizedX : CGFloat = touchLocation.x / viewBounds.width
-        if uiAspectRatio > videoAspectRatio {
-            normalizedX = (normalizedX - 0.5) * (uiAspectRatio / videoAspectRatio) + 0.5
-        }
-        
-        var normalizedY : CGFloat = touchLocation.y / viewBounds.height
-        if uiAspectRatio < videoAspectRatio {
-            normalizedY = (normalizedY - 0.5) * (videoAspectRatio / uiAspectRatio) + 0.5
-        }
-        
-        // normalize force value of touch
-        let normalizedForce : UInt8 = touch.maximumPossibleForce > 0 ? UInt8(touch.force / touch.maximumPossibleForce * CGFloat(UInt8.max)) : 1
-        
-        // Detect if touch is out of bounds
-        if normalizedX < 0.0 || normalizedX > 1.0 || normalizedY < 0.0 || normalizedY > 1.0 {
-            return TouchPoint(x: UInt16.max, y: UInt16.max, inRange: false, force: 0)
+        if let tv = touchView {
+            let uiAspectRatio : CGFloat = tv.bounds.width / tv.bounds.height
+            let touchLocation : CGPoint = touch.location(in: tv)
+            let viewBounds : CGRect = tv.bounds
+            
+            // normalise x,y to the UI element bounds
+            var normalizedX : CGFloat = touchLocation.x / viewBounds.width
+            if uiAspectRatio > videoAspectRatio {
+                normalizedX = (normalizedX - 0.5) * (uiAspectRatio / videoAspectRatio) + 0.5
+            }
+            
+            var normalizedY : CGFloat = touchLocation.y / viewBounds.height
+            if uiAspectRatio < videoAspectRatio {
+                normalizedY = (normalizedY - 0.5) * (videoAspectRatio / uiAspectRatio) + 0.5
+            }
+            
+            // normalize force value of touch
+            let normalizedForce : UInt8 = touch.maximumPossibleForce > 0 ? UInt8(touch.force / touch.maximumPossibleForce * CGFloat(UInt8.max)) : 1
+            
+            // Detect if touch is out of bounds
+            if normalizedX < 0.0 || normalizedX > 1.0 || normalizedY < 0.0 || normalizedY > 1.0 {
+                return TouchPoint(x: UInt16.max, y: UInt16.max, inRange: false, force: 0)
+            } else {
+                return TouchPoint(x: UInt16(normalizedX * CGFloat(UInt16.max)), y: UInt16(normalizedY * CGFloat(UInt16.max)), inRange: true, force: normalizedForce)
+            }
         } else {
-            return TouchPoint(x: UInt16(normalizedX * CGFloat(UInt16.max)), y: UInt16(normalizedY * CGFloat(UInt16.max)), inRange: true, force: normalizedForce)
+            return TouchPoint(x: UInt16.max, y: UInt16.max, inRange: false, force: 0)
         }
     }
     
     func sendTouchData(messageType: PixelStreamingToStreamerMessage, touches: Set<UITouch>) {
         
         guard relayTouchEvents else { return }
+        guard (webRTCClient != nil) else { return }
         
         var bytes: [UInt8] = []
         
@@ -190,7 +199,7 @@ final class TouchControls : TouchDelegate {
         }
         
         let data = Data(bytes)
-        self.webRTCClient.sendData(data)
+        self.webRTCClient?.sendData(data)
     }
     
 }
diff --git a/Engine/Extras/VirtualProduction/LiveLinkVCAM/vcam/WebRTC/Services/WebRTCClient.swift b/Engine/Extras/VirtualProduction/LiveLinkVCAM/vcam/WebRTC/Services/WebRTCClient.swift
--- a/Engine/Extras/VirtualProduction/LiveLinkVCAM/vcam/WebRTC/Services/WebRTCClient.swift
+++ b/Engine/Extras/VirtualProduction/LiveLinkVCAM/vcam/WebRTC/Services/WebRTCClient.swift
@@ -52,6 +52,10 @@ final class WebRTCClient: NSObject {
         super.init()
     }
     
+    deinit {
+        Log.info("WebRTCClient destructed.")
+    }
+    
     // MARK: PeerConnection Config
     func setupPeerConnection(rtcConfiguration: RTCConfiguration) {
         
@@ -68,6 +72,15 @@ final class WebRTCClient: NSObject {
         self.configureAudioSession()
         self.peerConnection!.delegate = self
     }
+    
+    func close() {
+        if let pc = self.peerConnection {
+            pc.close()
+            self.peerConnection = nil
+        }
+        self.remoteVideoTrack = nil
+        self.dataChannel = nil
+    }
 
     func stats(_ completionHandler : @escaping RTCStatisticsCompletionHandler) {
         self.peerConnection?.statistics(completionHandler: completionHandler)
diff --git a/Engine/Extras/VirtualProduction/LiveLinkVCAM/vcam/WebRTCStats.swift b/Engine/Extras/VirtualProduction/LiveLinkVCAM/vcam/WebRTCStats.swift
--- a/Engine/Extras/VirtualProduction/LiveLinkVCAM/vcam/WebRTCStats.swift
+++ b/Engine/Extras/VirtualProduction/LiveLinkVCAM/vcam/WebRTCStats.swift
@@ -19,6 +19,7 @@ class WebRTCStats
     private var lastBytesReceivedTimestamp : CFTimeInterval?
     private var lastInterframeDelay : Double?
     private var lastProcessingDelay : Double?
+    private var lastARKitEventTimestamp : CFTimeInterval?
     
     private var statsView : WebRTCStatsView
     private var fpsGraph : LineGraphView?
@@ -30,6 +31,7 @@ class WebRTCStats
     private var jitterBufferDelayGraph : LineGraphView?
     private var multipacketFramesGraph : LineGraphView?
     private var freezeCountGraph : LineGraphView?
+    private var pauseCountGraph : LineGraphView?
     private var nackCountGraph : LineGraphView?
     private var pliCountGraph : LineGraphView?
     private var packetLossGraph : LineGraphView?
@@ -38,6 +40,8 @@ class WebRTCStats
     private var framesDroppedGraph : LineGraphView?
     private var interframeDelayGraph : LineGraphView?
     private var processingTimeGraph : LineGraphView?
+    private var arkitEventsGraph : LineGraphView?
+    private var arkitEventsProcessedGraph : LineGraphView?
     
     init(statsView: WebRTCStatsView) {
         
@@ -69,6 +73,9 @@ class WebRTCStats
         self.freezeCountGraph = self.statsView.addGraph(graphName: "Freeze Count")
         self.freezeCountGraph?.maxY = 10
         
+        self.pauseCountGraph = self.statsView.addGraph(graphName: "Pause Count")
+        self.pauseCountGraph?.maxY = 10
+        
         self.pliCountGraph = self.statsView.addGraph(graphName: "Picture Loss Count")
         self.pliCountGraph?.maxY = 10
         
@@ -87,6 +94,11 @@ class WebRTCStats
         self.processingTimeGraph = self.statsView.addGraph(graphName: "Processing delay (ms)")
         self.processingTimeGraph?.maxY = 10
         
+        self.arkitEventsGraph = self.statsView.addGraph(graphName: "ARKit Transform Delta (ms)")
+        self.arkitEventsGraph?.maxY = 100
+        
+        self.arkitEventsProcessedGraph = self.statsView.addGraph(graphName: "ARKit Recv in UE / s")
+        self.arkitEventsProcessedGraph?.maxY = 100
     }
     
     func processStatsReport(report: RTCStatisticsReport){
@@ -129,6 +141,11 @@ class WebRTCStats
                         self.freezeCountGraph?.addDataPoint(value: Double(freezeCount))
                     }
                     
+                    // Pause Count
+                    if let pauseCount = value.values["pauseCount"] as? Int64 {
+                        self.pauseCountGraph?.addDataPoint(value: Double(pauseCount))
+                    }
+                    
                     // Nack Count
                     if let nackCount = value.values["nackCount"] as? Int64 {
                         self.nackCountGraph?.addDataPoint(value: Double(nackCount))
@@ -201,7 +218,6 @@ class WebRTCStats
 
                     // There is also:
                     // firCount
-                    // pauseCount
                     // totalSquaredInterFrameDelay : Double
                     // totalDecodeTime : Double
                     // totalPausesDuration : Double
@@ -213,8 +229,18 @@ class WebRTCStats
                 }
             }
         }
-        
-        
+    }
+    
+    func processARKitResponse(numResponses: UInt16){
+        self.arkitEventsProcessedGraph?.addDataPoint(value: Double(numResponses))
+    }
+    
+    func processARKitEvent(){
+        let now = CFAbsoluteTimeGetCurrent()
+        let prevTime = self.lastARKitEventTimestamp ?? CFAbsoluteTimeGetCurrent()
+        let eventDelta = now - prevTime
+        self.arkitEventsGraph?.addDataPoint(value: Double(eventDelta * 1000))
+        self.lastARKitEventTimestamp = now
     }
     
 }
diff --git a/Engine/Extras/VirtualProduction/LiveLinkVCAM/vcam/WebRTCStatsView.swift b/Engine/Extras/VirtualProduction/LiveLinkVCAM/vcam/WebRTCStatsView.swift
--- a/Engine/Extras/VirtualProduction/LiveLinkVCAM/vcam/WebRTCStatsView.swift
+++ b/Engine/Extras/VirtualProduction/LiveLinkVCAM/vcam/WebRTCStatsView.swift
@@ -34,21 +34,25 @@ class WebRTCStatsView: UIView {
     // Graphs will be displayed in a grid layout, rowsxcols, where nRows and nCols are determined by frameWidth/graphWidth
     // When graphs exceed nCols on a single row then a new row will begin
     // There is no maximum number of rows, but some graphs may be drawn offscreen if there is too many rows
-    private var nCols : Int = 5
+    private var nCols : Int = 6
+    private var nRows : Int = 3
     
     override init(frame: CGRect) {
         super.init(frame: frame)
         
+        self.graphLeftMargin = Int(frame.width / 100)
+        self.graphRightMargin = Int(frame.width / 100)
+        
+        let scale = abs(2 - UIScreen.main.scale) + 1
+        
+        self.graphWidth = Int(frame.width / (scale * CGFloat(self.nCols + 1)))
+        self.graphHeight = Int(frame.height / (scale * CGFloat(self.nRows + 1)))
+        
         // Disable user interaction so stats do not block underlying video
         self.isUserInteractionEnabled = false
         
         // set transparent background color
         self.backgroundColor = UIColor.clear
-        
-        // Calculate nRows and nCols
-        let frameWidth : CGFloat = frame.width - CGFloat(self.graphLeftMargin) - CGFloat(self.graphRightMargin)
-        self.nCols = Int( floor( frameWidth / CGFloat(self.graphWidth) ))
-
     }
     
     required init?(coder: NSCoder) {
diff --git a/Engine/Extras/VirtualProduction/LiveLinkVCAM/vcam/WebRTCStreamingConnection.swift b/Engine/Extras/VirtualProduction/LiveLinkVCAM/vcam/WebRTCStreamingConnection.swift
--- a/Engine/Extras/VirtualProduction/LiveLinkVCAM/vcam/WebRTCStreamingConnection.swift
+++ b/Engine/Extras/VirtualProduction/LiveLinkVCAM/vcam/WebRTCStreamingConnection.swift
@@ -17,9 +17,8 @@ class WebRTCStreamingConnection : StreamingConnection {
     private var signalClient: SignalingClient?
     private var touchControls: TouchControls?
     private var keyboardControls: KeyboardControls?
-    private var webRTCView : WebRTCView?
-    private var webRTCStatsView : WebRTCStatsView?
-    private var rtcVideoTrack : RTCVideoTrack?
+    private weak var webRTCView : WebRTCView?
+    private weak var rtcVideoTrack : RTCVideoTrack?
     
     private var signalingConnected = false
     private var hasRemoteSdp = false
@@ -66,6 +65,8 @@ class WebRTCStreamingConnection : StreamingConnection {
     override var renderView: UIView? {
         didSet {
             if let rv = renderView {
+                
+                // Attach webrtc video view to render view
                 let rtcView = WebRTCView(frame: CGRect(x: 0, y: 0, width: rv.frame.size.width, height: rv.frame.size.height))
                 rv.addSubview(rtcView)
                 rtcView.delegate = self
@@ -73,20 +74,16 @@ class WebRTCStreamingConnection : StreamingConnection {
                 self.webRTCView = rtcView
                 self.attachVideoTrack()
                 
-                // Attach stats view here
-                let rtcStatsView = WebRTCStatsView(frame: CGRect(x: 0, y: 0, width: rv.frame.size.width, height: rv.frame.size.height));
+                // Attach stats view to renderview
+                let rtcStatsView = WebRTCStatsView(frame: CGRect(x: 0, y: 0, width: rv.frame.size.width, height: rv.frame.size.height))
+                self.setupWebRTCStats(statsView: rtcStatsView)
                 rv.addSubview(rtcStatsView)
                 NSLayoutConstraint.activate([
                     rtcStatsView.topAnchor.constraint(equalTo: rv.topAnchor),
                     rtcStatsView.leadingAnchor.constraint(equalTo: rv.leadingAnchor),
                     rtcStatsView.trailingAnchor.constraint(equalTo: rv.trailingAnchor),
                     rtcStatsView.bottomAnchor.constraint(equalTo: rv.bottomAnchor)
-                ]);
-                // start with stats hidden
-                rtcStatsView.isHidden = true
-                self.webRTCStatsView = rtcStatsView;
-                self.webRTCStats = WebRTCStats(statsView: rtcStatsView)
-                
+                ])
             }
         }
     }
@@ -99,7 +96,19 @@ class WebRTCStreamingConnection : StreamingConnection {
         self.webRTCClient?.delegate = self
         
         self.stats = StreamingConnectionStats()
-
+    }
+    
+    deinit {
+        Log.info("WebRTCStreamingConnection destructed")
+        webRTCClient = nil
+    }
+    
+    func setupWebRTCStats(statsView: WebRTCStatsView) {
+        // Create stats object that drives the stats view
+        statsView.isHidden = true
+        self.webRTCStats = WebRTCStats(statsView: statsView)
+        
+        // Start a timer to update the stats
         _statsTimer = Timer.scheduledTimer(withTimeInterval: 1, repeats: true, block: { timer in
             
             if let webRTC = self.webRTCClient {
@@ -116,15 +125,14 @@ class WebRTCStreamingConnection : StreamingConnection {
         })
     }
     
-    deinit {
-        
-        Log.info("Destroyed WebRTCStreamingConnection")
-
-        webRTCClient = nil
-    }
-    
     override func showStats(_ shouldShow : Bool) {
-        self.webRTCStatsView?.isHidden = shouldShow == false
+        if let rv = self.renderView {
+            for subview in rv.subviews {
+                if let rtcStatsView = subview as? WebRTCStatsView {
+                    rtcStatsView.isHidden = !shouldShow
+                }
+            }
+        }
     }
 
     override func shutdown() {
@@ -175,10 +183,27 @@ class WebRTCStreamingConnection : StreamingConnection {
         hasRemoteSdp = false
         remoteCandidateCount = 0
         localCandidateCount = 0
+        
+        // stop the WebRTC client
+        self.webRTCClient?.close()
+        
+        // remove the video track from the webrtc view
+        self.detachVideoTrack()
+        
+        // Clear the stats timer
+        self._statsTimer?.invalidate()
+        
+        // Remove WebRTC view/WebRTC stats view from the render view on shutdown
+        if let rv = self.renderView {
+            for subview in rv.subviews {
+                subview.removeFromSuperview()
+            }
+        }
     }
     
     override func sendTransform(_ transform: simd_float4x4, atTime time: Double) {
-    
+        self.webRTCStats?.processARKitEvent()
+        
         guard let client = webRTCClient else { return }
     
         
@@ -344,6 +369,12 @@ class WebRTCStreamingConnection : StreamingConnection {
             view.attachTouchDelegate(delegate: self.touchControls!)
         }
     }
+    
+    func detachVideoTrack() {
+        if let view = self.webRTCView, let track = self.rtcVideoTrack {
+            view.removeVideoTrack(track: track)
+        }
+    }
 }
 
 extension WebRTCStreamingConnection : WebRTCViewDelegate {
@@ -583,7 +614,8 @@ extension WebRTCStreamingConnection: WebRTCClientDelegate {
                         self.webRTCClient?.sendRequestKeyFrame()
                     }
                 case .Response:
-                    fallthrough
+                    let numResponses : String? = String(data: data.dropFirst(), encoding: .utf16LittleEndian)
+                    self.webRTCStats?.processARKitResponse(numResponses: UInt16(numResponses ?? "") ?? 0)
                 case .FreezeFrame:
                     fallthrough
                 case .UnfreezeFrame:
diff --git a/Engine/Extras/VirtualProduction/LiveLinkVCAM/vcam/WebRTCView.swift b/Engine/Extras/VirtualProduction/LiveLinkVCAM/vcam/WebRTCView.swift
--- a/Engine/Extras/VirtualProduction/LiveLinkVCAM/vcam/WebRTCView.swift
+++ b/Engine/Extras/VirtualProduction/LiveLinkVCAM/vcam/WebRTCView.swift
@@ -9,63 +9,79 @@ import UIKit
 import WebRTC
 import MetalKit
 
-protocol WebRTCViewDelegate {
+protocol WebRTCViewDelegate : AnyObject {
     func webRTCView(_ view : WebRTCView, didChangeVideoSize size : CGSize)
 }
 
 final class WebRTCView: UIView, RTCVideoViewDelegate {
     
-    var delegate : WebRTCViewDelegate?
+    weak var delegate : WebRTCViewDelegate?
+    weak var videoView : RTCMTLVideoViewWithTouch?
     
-    let videoView = RTCMTLVideoViewWithTouch(frame: .zero)
     var videoSize = CGSize.zero
 
     override init(frame: CGRect) {
         super.init(frame: frame)
         
-        self.videoView.videoContentMode = .scaleAspectFit
-        self.videoView.isEnabled = true
-        self.videoView.isHidden = false
-        self.videoView.backgroundColor = UIColor.lightGray
+        let videoView = RTCMTLVideoViewWithTouch(frame: .zero)
+        videoView.videoContentMode = .scaleAspectFit
+        videoView.isEnabled = true
+        videoView.isHidden = false
+        videoView.backgroundColor = UIColor.lightGray
         
         videoView.delegate = self
         addSubview(videoView)
+        self.videoView = videoView
     }
 
     required init?(coder: NSCoder) {
         super.init(coder: coder)
         
+        let videoView = RTCMTLVideoViewWithTouch(frame: .zero)
         videoView.delegate = self
         addSubview(videoView)
+        self.videoView = videoView
+    }
+    
+    deinit {
+        Log.info("WebRTCView destructed.")
     }
 
     func videoView(_ videoView: RTCVideoRenderer, didChangeVideoSize size: CGSize) {
         self.videoSize = size
-        self.videoView.touchDelegate?.onVideoChangedSize(size: size)
+        self.videoView?.touchDelegate?.onVideoChangedSize(size: size)
         setNeedsLayout()
         
         self.delegate?.webRTCView(self, didChangeVideoSize: size)
     }
     
     func attachVideoTrack(track: RTCVideoTrack) {
-        track.remove(self.videoView)
-        track.add(self.videoView)
-        track.isEnabled = true
-        self.videoView.renderFrame(nil)
-        let trackState : RTCMediaStreamTrackState = track.readyState
-        if trackState == RTCMediaStreamTrackState.live {
-            print("Video track is live | id=\(track.trackId) | kind=\(track.kind)")
+        if let vidView = self.videoView {
+            track.remove(vidView)
+            track.add(vidView)
+            track.isEnabled = true
+            vidView.renderFrame(nil)
+            let trackState : RTCMediaStreamTrackState = track.readyState
+            if trackState == RTCMediaStreamTrackState.live {
+                print("Video track is live | id=\(track.trackId) | kind=\(track.kind)")
+            }
+        }
+    }
+    
+    func removeVideoTrack(track: RTCVideoTrack) {
+        if let vidView = self.videoView {
+            track.remove(vidView)
         }
     }
     
     func attachTouchDelegate(delegate: TouchDelegate) {
-        self.videoView.touchDelegate = delegate
+        self.videoView?.touchDelegate = delegate
     }
 
     override func layoutSubviews() {
         super.layoutSubviews()
         guard videoSize.width > 0 && videoSize.height > 0 else {
-            videoView.frame = bounds
+            videoView?.frame = bounds
             return
         }
 
@@ -74,8 +90,8 @@ final class WebRTCView: UIView, RTCVideoViewDelegate {
         videoFrame.size.width = videoFrame.size.width * CGFloat(scale)
         videoFrame.size.height = videoFrame.size.height * CGFloat(scale)
 
-        videoView.frame = videoFrame
-        videoView.center = CGPoint(x: bounds.midX, y: bounds.midY)
+        videoView?.frame = videoFrame
+        videoView?.center = CGPoint(x: bounds.midX, y: bounds.midY)
     }
     
 }
