[LiveLinkVCAM] Gamepad support

diff --git a/Engine/Extras/VirtualProduction/FlutterLiveLinkVCAM/IOS/Runner/AppDelegate.swift b/Engine/Extras/VirtualProduction/FlutterLiveLinkVCAM/IOS/Runner/AppDelegate.swift
--- a/Engine/Extras/VirtualProduction/FlutterLiveLinkVCAM/IOS/Runner/AppDelegate.swift
+++ b/Engine/Extras/VirtualProduction/FlutterLiveLinkVCAM/IOS/Runner/AppDelegate.swift
@@ -11,6 +11,7 @@ import WebRTC
   var videoViewControllerApi: FlutterRtcVideoViewControllerApi?
   var arSessionApi: FlutterArSessionApi?
   var tentacleApi: FlutterTentacleApi?
+  var gamepadApi: FlutterGamepadApi?
   
   override func application(
     _ application: UIApplication,
@@ -41,6 +42,7 @@ import WebRTC
     videoViewControllerApi = FlutterRtcVideoViewControllerApi(binaryMessenger: binaryMessenger)
     arSessionApi = FlutterArSessionApi(binaryMessenger: binaryMessenger)
     tentacleApi = FlutterTentacleApi(binaryMessenger: binaryMessenger)
+    gamepadApi = FlutterGamepadApi(binaryMessenger: binaryMessenger)
   
     return super.application(application, didFinishLaunchingWithOptions: launchOptions)
   }
diff --git a/Engine/Extras/VirtualProduction/FlutterLiveLinkVCAM/IOS/Runner/Gamepad/API/FlutterGamepadApi.swift b/Engine/Extras/VirtualProduction/FlutterLiveLinkVCAM/IOS/Runner/Gamepad/API/FlutterGamepadApi.swift
new file mode 100644
--- /dev/null
+++ b/Engine/Extras/VirtualProduction/FlutterLiveLinkVCAM/IOS/Runner/Gamepad/API/FlutterGamepadApi.swift
@@ -0,0 +1,211 @@
+// Copyright Epic Games, Inc. All Rights Reserved.
+
+import Foundation
+import GameController
+
+/// Handles messages about Tentacle to and from the Flutter API.
+class FlutterGamepadApi: FlutterPluginApi<GamepadFlutterApi> {
+  /// List of IDs for actively connected gamepads.
+  private var activeGamepadIds: Set<GCControllerPlayerIndex> = []
+  
+  init(binaryMessenger: FlutterBinaryMessenger) {
+    super.init(flutter:GamepadFlutterApi(binaryMessenger: binaryMessenger))
+    
+    GamepadHostApiSetup.setUp(binaryMessenger: binaryMessenger, api: self)
+    FlutterGamepadApi.instance = self
+    
+    // Check for already connected controllers
+    for controller in GCController.controllers() {
+      onControllerConnected(gamepad: controller)
+    }
+
+    let notificationCenter = NotificationCenter.default
+    
+    // Add notifications for when new controllers connect or disconnect
+    notificationCenter.addObserver(forName: .GCControllerDidConnect, object: nil, queue: .main) { (note) in
+      guard let controller = note.object as? GCController else { return }
+      self.onControllerConnected(gamepad: controller)
+    }
+    
+    notificationCenter.addObserver(forName: .GCControllerDidDisconnect, object: nil, queue: .main) { (note) in
+      guard let controller = note.object as? GCController else { return }
+      self.onControllerDisconnected(gamepad: controller)
+    }
+  }
+  
+  deinit {
+    if (FlutterGamepadApi.instance === self) {
+      FlutterGamepadApi.instance = nil
+    }
+    
+    let notificationCenter = NotificationCenter.default
+    notificationCenter.removeObserver(self, name: .GCControllerDidConnect, object: nil)
+    notificationCenter.removeObserver(self, name: .GCControllerDidDisconnect, object: nil)
+  }
+  
+  /// The active singleton instance of this class.
+  private(set) static var instance: FlutterGamepadApi?
+  
+  /// Called when a controller is connected.
+  private func onControllerConnected(gamepad: GCController) {
+    activeGamepadIds.insert(gamepad.playerIndex)
+    
+    gamepad.extendedGamepad?.valueChangedHandler = self.onValueChanged
+    
+    callFlutter { flutter in
+      flutter.onGamepadConnected(gamepadId: Int64(gamepad.playerIndex.rawValue)) { _ in }
+    }
+  }
+  
+  /// Called when a controller is disconnected.
+  private func onControllerDisconnected(gamepad: GCController) {
+    activeGamepadIds.remove(gamepad.playerIndex)
+    
+    callFlutter { flutter in
+      flutter.onGamepadDisconnected(gamepadId: Int64(gamepad.playerIndex.rawValue)) { _ in }
+    }
+  }
+  
+  /// Called when a controller value changes.
+  private func onValueChanged(gamepad: GCExtendedGamepad, element: GCControllerElement) {
+    let playerIndex: GCControllerPlayerIndex? = gamepad.controller?.playerIndex
+    
+    if (playerIndex == nil) {
+      return
+    }
+    
+    // Handle button input
+    if (element is GCControllerButtonInput) {
+      onButtonValueChanged(gamepad: gamepad, button: element as! GCControllerButtonInput)
+      return
+    }
+    
+    // Handle simple axis input
+    if (element is GCControllerAxisInput) {
+      onAxisValueChanged(gamepad: gamepad, axis: element as! GCControllerAxisInput)
+      return
+    }
+    
+    // Special handling for inputs that fire multiple sub-events
+    switch (element) {
+    case gamepad.dpad:
+      onButtonValueChanged(gamepad: gamepad, button: gamepad.dpad.up)
+      onButtonValueChanged(gamepad: gamepad, button: gamepad.dpad.down)
+      onButtonValueChanged(gamepad: gamepad, button: gamepad.dpad.left)
+      onButtonValueChanged(gamepad: gamepad, button: gamepad.dpad.right)
+      return
+      
+    case gamepad.leftThumbstick:
+      onAxisValueChanged(gamepad: gamepad, axis: gamepad.leftThumbstick.xAxis)
+      onAxisValueChanged(gamepad: gamepad, axis: gamepad.leftThumbstick.yAxis)
+      return
+      
+    case gamepad.rightThumbstick:
+      onAxisValueChanged(gamepad: gamepad, axis: gamepad.rightThumbstick.xAxis)
+      onAxisValueChanged(gamepad: gamepad, axis: gamepad.rightThumbstick.yAxis)
+      return
+      
+    default:
+      break
+    }
+  }
+  
+  /// Called when an axis value on the gamepad changes.
+  private func onAxisValueChanged(gamepad: GCExtendedGamepad, axis: GCControllerAxisInput) {
+    let playerIndex: Int? = gamepad.controller?.playerIndex.rawValue
+    if (playerIndex == nil) {
+      return
+    }
+    
+    let input: GamepadInput? = switch (axis) {
+    case gamepad.leftThumbstick.xAxis: .thumbstickLeftX
+    case gamepad.leftThumbstick.yAxis: .thumbstickLeftY
+    case gamepad.rightThumbstick.xAxis: .thumbstickRightX
+    case gamepad.rightThumbstick.yAxis: .thumbstickRightY
+    case gamepad.leftTrigger: .triggerAxisLeft
+    case gamepad.rightTrigger: .triggerAxisRight
+    default: nil
+    }
+    
+    if (input == nil) {
+      return
+    }
+    
+    callFlutter { flutter in
+      flutter.onGamepadInputEvent(event: GamepadInputEvent(
+        gamepadId: Int64(playerIndex!),
+        input: input!,
+        type: GamepadInputType.axis,
+        value: Double(axis.value)
+      )) { _ in }
+    }
+  }
+  
+  /// Called when a button's value on the gamepad changes.
+  private func onButtonValueChanged(gamepad: GCExtendedGamepad, button: GCControllerButtonInput) {
+    let playerIndex: Int? = gamepad.controller?.playerIndex.rawValue
+    if (playerIndex == nil) {
+      return
+    }
+    
+    let input: GamepadInput? = switch (button) {
+    case gamepad.dpad.up: .dpadUp
+    case gamepad.dpad.down: .dpadDown
+    case gamepad.dpad.left: .dpadLeft
+    case gamepad.dpad.right: .dpadRight
+    case gamepad.buttonA: .faceButtonBottom
+    case gamepad.buttonB: .faceButtonRight
+    case gamepad.buttonX: .faceButtonLeft
+    case gamepad.buttonY: .faceButtonRight
+    case gamepad.leftShoulder: .shoulderButtonLeft
+    case gamepad.rightShoulder: .shoulderButtonRight
+    case gamepad.leftTrigger: .triggerButtonLeft
+    case gamepad.rightTrigger: .triggerButtonRight
+    case gamepad.leftThumbstickButton: .thumbstickLeftButton
+    case gamepad.rightThumbstickButton: .thumbstickRightButton
+    case gamepad.buttonOptions: .specialButtonLeft
+    case gamepad.buttonMenu: .specialButtonRight
+    default: nil
+    }
+    
+    if (input == nil) {
+      return
+    }
+    
+    callFlutter { flutter in
+      flutter.onGamepadInputEvent(event: GamepadInputEvent(
+        gamepadId: Int64(playerIndex!),
+        input: input!,
+        type: GamepadInputType.button,
+        value: button.isPressed ? 1 : 0
+      )) { _ in }
+    }
+    
+    // Fire an axis event as well for triggers, which act as both buttons and axes
+    let axisInput: GamepadInput? = switch (button) {
+    case gamepad.leftTrigger: .triggerAxisLeft
+    case gamepad.rightTrigger: .triggerAxisRight
+    default: nil
+    }
+    
+    if (axisInput == nil) {
+      return
+    }
+    
+    callFlutter { flutter in
+      flutter.onGamepadInputEvent(event: GamepadInputEvent(
+        gamepadId: Int64(playerIndex!),
+        input: axisInput!,
+        type: GamepadInputType.axis,
+        value: Double(button.value)
+      )) { _ in }
+    }
+  }
+}
+
+// MARK: Host API
+extension FlutterGamepadApi: GamepadHostApi {
+  func getActiveGamepadIds() throws -> [Int64] {
+    return activeGamepadIds.map { playerIndex in Int64(playerIndex.rawValue) }
+  }
+}
diff --git a/Engine/Extras/VirtualProduction/FlutterLiveLinkVCAM/IOS/Runner/Tentacle/API/FlutterTentacleApi.swift b/Engine/Extras/VirtualProduction/FlutterLiveLinkVCAM/IOS/Runner/Tentacle/API/FlutterTentacleApi.swift
--- a/Engine/Extras/VirtualProduction/FlutterLiveLinkVCAM/IOS/Runner/Tentacle/API/FlutterTentacleApi.swift
+++ b/Engine/Extras/VirtualProduction/FlutterLiveLinkVCAM/IOS/Runner/Tentacle/API/FlutterTentacleApi.swift
@@ -1,6 +1,5 @@
 // Copyright Epic Games, Inc. All Rights Reserved.
 
-import ARKit
 import Foundation
 import Tentacle
 
@@ -23,7 +22,7 @@ class FlutterTentacleApi: FlutterPluginApi<TentacleFlutterApi> {
   }
   
   deinit {
-    if (FlutterArSessionApi.instance === self) {
+    if (FlutterTentacleApi.instance === self) {
       FlutterTentacleApi.instance = nil
     }
   }
