[LiveLinkVCAM] Add missing file

diff --git a/Engine/Extras/VirtualProduction/FlutterLiveLinkVCAM/IOS/Runner/Util/PigeonExtension.swift b/Engine/Extras/VirtualProduction/FlutterLiveLinkVCAM/IOS/Runner/Util/PigeonExtension.swift
new file mode 100644
--- /dev/null
+++ b/Engine/Extras/VirtualProduction/FlutterLiveLinkVCAM/IOS/Runner/Util/PigeonExtension.swift
@@ -0,0 +1,13 @@
+// Copyright Epic Games, Inc. All Rights Reserved.
+
+import Foundation
+#if os(iOS)
+import Flutter
+#elseif os(macOS)
+import FlutterMacOS
+#else
+#error("Unsupported platform.")
+#endif
+
+/// Workaround for compilation errors originating from Pigeon error handling code.
+extension FlutterError: Swift.Error {}
