[LiveLinkVCAM] Use texture instead of platform view for WebRTC stream #jira UE-205525

diff --git a/Engine/Extras/VirtualProduction/FlutterLiveLinkVCAM/IOS/Runner/AppDelegate.swift b/Engine/Extras/VirtualProduction/FlutterLiveLinkVCAM/IOS/Runner/AppDelegate.swift
--- a/Engine/Extras/VirtualProduction/FlutterLiveLinkVCAM/IOS/Runner/AppDelegate.swift
+++ b/Engine/Extras/VirtualProduction/FlutterLiveLinkVCAM/IOS/Runner/AppDelegate.swift
@@ -3,6 +3,7 @@
 import UIKit
 import Flutter
 import WebRTC
+import SpriteKit
 
 @UIApplicationMain
 @objc class AppDelegate: FlutterAppDelegate {
@@ -21,25 +22,13 @@ import WebRTC
     
     GeneratedPluginRegistrant.register(with: self)
     
-    // Register RtcVideoView platform view
-    let videoViewRegistrar: FlutterPluginRegistrar? = self.registrar(forPlugin: "RtcVideoView")
-    if (videoViewRegistrar == nil) {
-      fatalError("Failed to register RtcVideoView plugin")
-    }
-    
-    videoViewRegistrar?.register(
-      RtcVideoViewFactory(
-        messenger: videoViewRegistrar!.messenger()
-      ),
-      withId: "com.epicgames.live_link_vcam.view.RtcVideoView"
-    )
-    
     // Register Pigeon APIs
-    let binaryMessenger = (window?.rootViewController as! FlutterViewController).binaryMessenger
+    let flutterViewController = (window?.rootViewController as! FlutterViewController)
+    let binaryMessenger = flutterViewController.binaryMessenger
     
     peerConnectionApi = FlutterRtcPeerConnectionApi(binaryMessenger: binaryMessenger)
     dataChannelApi = FlutterRtcDataChannelApi(binaryMessenger: binaryMessenger)
-    videoViewControllerApi = FlutterRtcVideoViewControllerApi(binaryMessenger: binaryMessenger)
+    videoViewControllerApi = FlutterRtcVideoViewControllerApi(binaryMessenger: binaryMessenger, textureRegistry: flutterViewController)
     arSessionApi = FlutterArSessionApi(binaryMessenger: binaryMessenger)
     tentacleApi = FlutterTentacleApi(binaryMessenger: binaryMessenger)
     gamepadApi = FlutterGamepadApi(binaryMessenger: binaryMessenger)
diff --git a/Engine/Extras/VirtualProduction/FlutterLiveLinkVCAM/IOS/Runner/View/RtcVideoView.swift b/Engine/Extras/VirtualProduction/FlutterLiveLinkVCAM/IOS/Runner/View/RtcVideoView.swift
deleted file mode 100644
diff --git a/Engine/Extras/VirtualProduction/FlutterLiveLinkVCAM/IOS/Runner/View/RtcVideoViewFactory.swift b/Engine/Extras/VirtualProduction/FlutterLiveLinkVCAM/IOS/Runner/View/RtcVideoViewFactory.swift
deleted file mode 100644
diff --git a/Engine/Extras/VirtualProduction/FlutterLiveLinkVCAM/IOS/Runner/WebRTC/API/FlutterRtcVideoViewControllerApi.swift b/Engine/Extras/VirtualProduction/FlutterLiveLinkVCAM/IOS/Runner/WebRTC/API/FlutterRtcVideoViewControllerApi.swift
--- a/Engine/Extras/VirtualProduction/FlutterLiveLinkVCAM/IOS/Runner/WebRTC/API/FlutterRtcVideoViewControllerApi.swift
+++ b/Engine/Extras/VirtualProduction/FlutterLiveLinkVCAM/IOS/Runner/WebRTC/API/FlutterRtcVideoViewControllerApi.swift
@@ -1,10 +1,13 @@
 // Copyright Epic Games, Inc. All Rights Reserved.
 
+import Flutter
 import Foundation
 
 /// Handles messages about [FlutterRtcPeerConnection]s to and from the Flutter API.
 class FlutterRtcVideoViewControllerApi: FlutterPluginApi<RtcVideoViewControllerFlutterApi> {
-  init(binaryMessenger: FlutterBinaryMessenger) {
+  init(binaryMessenger: FlutterBinaryMessenger, textureRegistry: FlutterTextureRegistry) {
+    self.textureRegistry = textureRegistry
+    
     super.init(flutter: RtcVideoViewControllerFlutterApi(binaryMessenger: binaryMessenger))
     
     RtcVideoViewControllerHostApiSetup.setUp(binaryMessenger: binaryMessenger, api: self)
@@ -20,6 +23,9 @@ class FlutterRtcVideoViewControllerApi: FlutterPluginApi<RtcVideoViewControllerF
   /// Manager for indexed video view controllers shared with Flutter.
   let videoViewControllerManager = IdObjectManager<FlutterRtcVideoViewController>(debugTypeName: "VideoViewController")
   
+  /// Texture registry used to communicate with Flutter.
+  let textureRegistry: FlutterTextureRegistry
+  
   /// The active singleton instance of this class.
   private(set) static var instance: FlutterRtcVideoViewControllerApi?
 }
@@ -46,6 +52,19 @@ extension FlutterRtcVideoViewControllerApi: RtcVideoViewControllerHostApi {
       .setTrack(newTrack: track)
   }
   
+  func getTextureId(controllerId: Int64) throws -> Int64 {
+    return videoViewControllerManager
+      .getChecked(id: controllerId)
+      .textureId
+  }
+  
+  func clear(controllerId: Int64) throws {
+    videoViewControllerManager
+      .getChecked(id: controllerId)
+      .clear()
+  }
+  
+  
   func dispose(controllerId: Int64) throws {
     videoViewControllerManager.unregister(id: controllerId)
   }
diff --git a/Engine/Extras/VirtualProduction/FlutterLiveLinkVCAM/IOS/Runner/WebRTC/API/WebRtcApi.g.swift b/Engine/Extras/VirtualProduction/FlutterLiveLinkVCAM/IOS/Runner/WebRTC/API/WebRtcApi.g.swift
--- a/Engine/Extras/VirtualProduction/FlutterLiveLinkVCAM/IOS/Runner/WebRTC/API/WebRtcApi.g.swift
+++ b/Engine/Extras/VirtualProduction/FlutterLiveLinkVCAM/IOS/Runner/WebRTC/API/WebRtcApi.g.swift
@@ -3,7 +3,7 @@
 // This is an auto-generated file. See FlutterLiveLinkVCAM/pigeons/README for information on how to
 // modify the original source and generate a new file.
 // 
-// Autogenerated from Pigeon (v11.0.1), do not edit directly.
+// Autogenerated from Pigeon (v14.0.1), do not edit directly.
 // See also: https://pub.dev/packages/pigeon
 
 import Foundation
@@ -15,10 +15,6 @@ import FlutterMacOS
 #error("Unsupported platform.")
 #endif
 
-private func isNullish(_ value: Any?) -> Bool {
-  return value is NSNull || value == nil
-}
-
 private func wrapResult(_ result: Any?) -> [Any?] {
   return [result]
 }
@@ -38,6 +34,14 @@ private func wrapError(_ error: Any) -> [Any?] {
   ]
 }
 
+private func createConnectionError(withChannelName channelName: String) -> FlutterError {
+  return FlutterError(code: "channel-error", message: "Unable to establish connection on channel: '\(channelName)'.", details: "")
+}
+
+private func isNullish(_ value: Any?) -> Bool {
+  return value is NSNull || value == nil
+}
+
 private func nilOrValue<T>(_ value: Any?) -> T? {
   if value is NSNull { return nil }
   return value as! T?
@@ -516,8 +520,20 @@ class RtcPeerConnectionFlutterApiCodec: FlutterStandardMessageCodec {
 
 /// API that receives messages about RtcPeerConnection instances in Flutter.
 ///
-/// Generated class from Pigeon that represents Flutter messages that can be called from Swift.
-class RtcPeerConnectionFlutterApi {
+/// Generated protocol from Pigeon that represents Flutter messages that can be called from Swift.
+protocol RtcPeerConnectionFlutterApiProtocol {
+  /// Called when a new local ICE [candidate] is found for the PeerConnection with the given [connectionId].
+  func onIceCandidate(connectionId connectionIdArg: Int64, candidate candidateArg: RtcIceCandidate?, completion: @escaping (Result<Void, FlutterError>) -> Void)
+  /// Called when the [state] of the PeerConnection with the given [connectionId] changes.
+  func onStateChanged(connectionId connectionIdArg: Int64, state stateArg: RtcPeerConnectionState, completion: @escaping (Result<Void, FlutterError>) -> Void)
+  /// Called when a MediaStreamTrack (described in [event]) is added to the PeerConnection with the given
+  /// [connectionId].
+  func onTrack(connectionId connectionIdArg: Int64, event eventArg: RtcTrackEvent, completion: @escaping (Result<Void, FlutterError>) -> Void)
+  /// Called when a DataChannel with the given [dataChannelId] is added to the PeerConnection with the given
+  /// [connectionId].
+  func onDataChannel(connectionId connectionIdArg: Int64, dataChannelId dataChannelIdArg: Int64, completion: @escaping (Result<Void, FlutterError>) -> Void)
+}
+class RtcPeerConnectionFlutterApi: RtcPeerConnectionFlutterApiProtocol {
   private let binaryMessenger: FlutterBinaryMessenger
   init(binaryMessenger: FlutterBinaryMessenger){
     self.binaryMessenger = binaryMessenger
@@ -526,33 +542,81 @@ class RtcPeerConnectionFlutterApi {
     return RtcPeerConnectionFlutterApiCodec.shared
   }
   /// Called when a new local ICE [candidate] is found for the PeerConnection with the given [connectionId].
-  func onIceCandidate(connectionId connectionIdArg: Int64, candidate candidateArg: RtcIceCandidate?, completion: @escaping () -> Void) {
-    let channel = FlutterBasicMessageChannel(name: "dev.flutter.pigeon.live_link_vcam.RtcPeerConnectionFlutterApi.onIceCandidate", binaryMessenger: binaryMessenger, codec: codec)
-    channel.sendMessage([connectionIdArg, candidateArg] as [Any?]) { _ in
-      completion()
+  func onIceCandidate(connectionId connectionIdArg: Int64, candidate candidateArg: RtcIceCandidate?, completion: @escaping (Result<Void, FlutterError>) -> Void) {
+    let channelName: String = "dev.flutter.pigeon.live_link_vcam.RtcPeerConnectionFlutterApi.onIceCandidate"
+    let channel = FlutterBasicMessageChannel(name: channelName, binaryMessenger: binaryMessenger, codec: codec)
+    channel.sendMessage([connectionIdArg, candidateArg] as [Any?]) { response in
+      guard let listResponse = response as? [Any?] else {
+        completion(.failure(createConnectionError(withChannelName:channelName)))
+        return
+      }
+      if (listResponse.count > 1) {
+        let code: String = listResponse[0] as! String
+        let message: String? = nilOrValue(listResponse[1])
+        let details: String? = nilOrValue(listResponse[2])
+        completion(.failure(FlutterError(code: code, message: message, details: details)));
+      } else {
+        completion(.success(Void()))
+      }
     }
   }
   /// Called when the [state] of the PeerConnection with the given [connectionId] changes.
-  func onStateChanged(connectionId connectionIdArg: Int64, state stateArg: RtcPeerConnectionState, completion: @escaping () -> Void) {
-    let channel = FlutterBasicMessageChannel(name: "dev.flutter.pigeon.live_link_vcam.RtcPeerConnectionFlutterApi.onStateChanged", binaryMessenger: binaryMessenger, codec: codec)
-    channel.sendMessage([connectionIdArg, stateArg.rawValue] as [Any?]) { _ in
-      completion()
+  func onStateChanged(connectionId connectionIdArg: Int64, state stateArg: RtcPeerConnectionState, completion: @escaping (Result<Void, FlutterError>) -> Void) {
+    let channelName: String = "dev.flutter.pigeon.live_link_vcam.RtcPeerConnectionFlutterApi.onStateChanged"
+    let channel = FlutterBasicMessageChannel(name: channelName, binaryMessenger: binaryMessenger, codec: codec)
+    channel.sendMessage([connectionIdArg, stateArg.rawValue] as [Any?]) { response in
+      guard let listResponse = response as? [Any?] else {
+        completion(.failure(createConnectionError(withChannelName:channelName)))
+        return
+      }
+      if (listResponse.count > 1) {
+        let code: String = listResponse[0] as! String
+        let message: String? = nilOrValue(listResponse[1])
+        let details: String? = nilOrValue(listResponse[2])
+        completion(.failure(FlutterError(code: code, message: message, details: details)));
+      } else {
+        completion(.success(Void()))
+      }
     }
   }
   /// Called when a MediaStreamTrack (described in [event]) is added to the PeerConnection with the given
   /// [connectionId].
-  func onTrack(connectionId connectionIdArg: Int64, event eventArg: RtcTrackEvent, completion: @escaping () -> Void) {
-    let channel = FlutterBasicMessageChannel(name: "dev.flutter.pigeon.live_link_vcam.RtcPeerConnectionFlutterApi.onTrack", binaryMessenger: binaryMessenger, codec: codec)
-    channel.sendMessage([connectionIdArg, eventArg] as [Any?]) { _ in
-      completion()
+  func onTrack(connectionId connectionIdArg: Int64, event eventArg: RtcTrackEvent, completion: @escaping (Result<Void, FlutterError>) -> Void) {
+    let channelName: String = "dev.flutter.pigeon.live_link_vcam.RtcPeerConnectionFlutterApi.onTrack"
+    let channel = FlutterBasicMessageChannel(name: channelName, binaryMessenger: binaryMessenger, codec: codec)
+    channel.sendMessage([connectionIdArg, eventArg] as [Any?]) { response in
+      guard let listResponse = response as? [Any?] else {
+        completion(.failure(createConnectionError(withChannelName:channelName)))
+        return
+      }
+      if (listResponse.count > 1) {
+        let code: String = listResponse[0] as! String
+        let message: String? = nilOrValue(listResponse[1])
+        let details: String? = nilOrValue(listResponse[2])
+        completion(.failure(FlutterError(code: code, message: message, details: details)));
+      } else {
+        completion(.success(Void()))
+      }
     }
   }
   /// Called when a DataChannel with the given [dataChannelId] is added to the PeerConnection with the given
   /// [connectionId].
-  func onDataChannel(connectionId connectionIdArg: Int64, dataChannelId dataChannelIdArg: Int64, completion: @escaping () -> Void) {
-    let channel = FlutterBasicMessageChannel(name: "dev.flutter.pigeon.live_link_vcam.RtcPeerConnectionFlutterApi.onDataChannel", binaryMessenger: binaryMessenger, codec: codec)
-    channel.sendMessage([connectionIdArg, dataChannelIdArg] as [Any?]) { _ in
-      completion()
+  func onDataChannel(connectionId connectionIdArg: Int64, dataChannelId dataChannelIdArg: Int64, completion: @escaping (Result<Void, FlutterError>) -> Void) {
+    let channelName: String = "dev.flutter.pigeon.live_link_vcam.RtcPeerConnectionFlutterApi.onDataChannel"
+    let channel = FlutterBasicMessageChannel(name: channelName, binaryMessenger: binaryMessenger, codec: codec)
+    channel.sendMessage([connectionIdArg, dataChannelIdArg] as [Any?]) { response in
+      guard let listResponse = response as? [Any?] else {
+        completion(.failure(createConnectionError(withChannelName:channelName)))
+        return
+      }
+      if (listResponse.count > 1) {
+        let code: String = listResponse[0] as! String
+        let message: String? = nilOrValue(listResponse[1])
+        let details: String? = nilOrValue(listResponse[2])
+        completion(.failure(FlutterError(code: code, message: message, details: details)));
+      } else {
+        completion(.success(Void()))
+      }
     }
   }
 }
@@ -663,8 +727,14 @@ class RtcDataChannelFlutterApiCodec: FlutterStandardMessageCodec {
 
 /// API that receives messages about RtcDataChannel instances in Flutter.
 ///
-/// Generated class from Pigeon that represents Flutter messages that can be called from Swift.
-class RtcDataChannelFlutterApi {
+/// Generated protocol from Pigeon that represents Flutter messages that can be called from Swift.
+protocol RtcDataChannelFlutterApiProtocol {
+  /// Called when the [state] of the DataChannel with the given [dataChannelId] changes.
+  func onStateChanged(dataChannelId dataChannelIdArg: Int64, state stateArg: RtcDataChannelState, completion: @escaping (Result<Void, FlutterError>) -> Void)
+  /// Called when a message contained in [buffer] is received on the data channel with the given [dataChannelId].
+  func onMessage(dataChannelId dataChannelIdArg: Int64, buffer bufferArg: RtcDataBuffer, completion: @escaping (Result<Void, FlutterError>) -> Void)
+}
+class RtcDataChannelFlutterApi: RtcDataChannelFlutterApiProtocol {
   private let binaryMessenger: FlutterBinaryMessenger
   init(binaryMessenger: FlutterBinaryMessenger){
     self.binaryMessenger = binaryMessenger
@@ -673,17 +743,41 @@ class RtcDataChannelFlutterApi {
     return RtcDataChannelFlutterApiCodec.shared
   }
   /// Called when the [state] of the DataChannel with the given [dataChannelId] changes.
-  func onStateChanged(dataChannelId dataChannelIdArg: Int64, state stateArg: RtcDataChannelState, completion: @escaping () -> Void) {
-    let channel = FlutterBasicMessageChannel(name: "dev.flutter.pigeon.live_link_vcam.RtcDataChannelFlutterApi.onStateChanged", binaryMessenger: binaryMessenger, codec: codec)
-    channel.sendMessage([dataChannelIdArg, stateArg.rawValue] as [Any?]) { _ in
-      completion()
+  func onStateChanged(dataChannelId dataChannelIdArg: Int64, state stateArg: RtcDataChannelState, completion: @escaping (Result<Void, FlutterError>) -> Void) {
+    let channelName: String = "dev.flutter.pigeon.live_link_vcam.RtcDataChannelFlutterApi.onStateChanged"
+    let channel = FlutterBasicMessageChannel(name: channelName, binaryMessenger: binaryMessenger, codec: codec)
+    channel.sendMessage([dataChannelIdArg, stateArg.rawValue] as [Any?]) { response in
+      guard let listResponse = response as? [Any?] else {
+        completion(.failure(createConnectionError(withChannelName:channelName)))
+        return
+      }
+      if (listResponse.count > 1) {
+        let code: String = listResponse[0] as! String
+        let message: String? = nilOrValue(listResponse[1])
+        let details: String? = nilOrValue(listResponse[2])
+        completion(.failure(FlutterError(code: code, message: message, details: details)));
+      } else {
+        completion(.success(Void()))
+      }
     }
   }
   /// Called when a message contained in [buffer] is received on the data channel with the given [dataChannelId].
-  func onMessage(dataChannelId dataChannelIdArg: Int64, buffer bufferArg: RtcDataBuffer, completion: @escaping () -> Void) {
-    let channel = FlutterBasicMessageChannel(name: "dev.flutter.pigeon.live_link_vcam.RtcDataChannelFlutterApi.onMessage", binaryMessenger: binaryMessenger, codec: codec)
-    channel.sendMessage([dataChannelIdArg, bufferArg] as [Any?]) { _ in
-      completion()
+  func onMessage(dataChannelId dataChannelIdArg: Int64, buffer bufferArg: RtcDataBuffer, completion: @escaping (Result<Void, FlutterError>) -> Void) {
+    let channelName: String = "dev.flutter.pigeon.live_link_vcam.RtcDataChannelFlutterApi.onMessage"
+    let channel = FlutterBasicMessageChannel(name: channelName, binaryMessenger: binaryMessenger, codec: codec)
+    channel.sendMessage([dataChannelIdArg, bufferArg] as [Any?]) { response in
+      guard let listResponse = response as? [Any?] else {
+        completion(.failure(createConnectionError(withChannelName:channelName)))
+        return
+      }
+      if (listResponse.count > 1) {
+        let code: String = listResponse[0] as! String
+        let message: String? = nilOrValue(listResponse[1])
+        let details: String? = nilOrValue(listResponse[2])
+        completion(.failure(FlutterError(code: code, message: message, details: details)));
+      } else {
+        completion(.success(Void()))
+      }
     }
   }
 }
@@ -697,6 +791,10 @@ protocol RtcVideoViewControllerHostApi {
   /// Set the video track of the VideoViewController with the given [controllerId] to the track with the given
   /// [trackId].
   func setTrack(controllerId: Int64, trackId: Int64) throws
+  /// Get the texture ID associated with the VideoViewController with the given [controllerId].
+  func getTextureId(controllerId: Int64) throws -> Int64
+  /// Clear the image of the VideoViewController with the given [controllerId].
+  func clear(controllerId: Int64) throws
   /// Dispose of the VideoViewController with the given [controllerId].
   func dispose(controllerId: Int64) throws
 }
@@ -739,6 +837,38 @@ class RtcVideoViewControllerHostApiSetup {
     } else {
       setTrackChannel.setMessageHandler(nil)
     }
+    /// Get the texture ID associated with the VideoViewController with the given [controllerId].
+    let getTextureIdChannel = FlutterBasicMessageChannel(name: "dev.flutter.pigeon.live_link_vcam.RtcVideoViewControllerHostApi.getTextureId", binaryMessenger: binaryMessenger)
+    if let api = api {
+      getTextureIdChannel.setMessageHandler { message, reply in
+        let args = message as! [Any?]
+        let controllerIdArg = args[0] is Int64 ? args[0] as! Int64 : Int64(args[0] as! Int32)
+        do {
+          let result = try api.getTextureId(controllerId: controllerIdArg)
+          reply(wrapResult(result))
+        } catch {
+          reply(wrapError(error))
+        }
+      }
+    } else {
+      getTextureIdChannel.setMessageHandler(nil)
+    }
+    /// Clear the image of the VideoViewController with the given [controllerId].
+    let clearChannel = FlutterBasicMessageChannel(name: "dev.flutter.pigeon.live_link_vcam.RtcVideoViewControllerHostApi.clear", binaryMessenger: binaryMessenger)
+    if let api = api {
+      clearChannel.setMessageHandler { message, reply in
+        let args = message as! [Any?]
+        let controllerIdArg = args[0] is Int64 ? args[0] as! Int64 : Int64(args[0] as! Int32)
+        do {
+          try api.clear(controllerId: controllerIdArg)
+          reply(wrapResult(nil))
+        } catch {
+          reply(wrapError(error))
+        }
+      }
+    } else {
+      clearChannel.setMessageHandler(nil)
+    }
     /// Dispose of the VideoViewController with the given [controllerId].
     let disposeChannel = FlutterBasicMessageChannel(name: "dev.flutter.pigeon.live_link_vcam.RtcVideoViewControllerHostApi.dispose", binaryMessenger: binaryMessenger)
     if let api = api {
@@ -759,17 +889,56 @@ class RtcVideoViewControllerHostApiSetup {
 }
 /// API that receives messages about VideoViewController instances in Flutter.
 ///
-/// Generated class from Pigeon that represents Flutter messages that can be called from Swift.
-class RtcVideoViewControllerFlutterApi {
+/// Generated protocol from Pigeon that represents Flutter messages that can be called from Swift.
+protocol RtcVideoViewControllerFlutterApiProtocol {
+  /// Called when the VideoViewController with the given [controllerId] has rendered its first frame since it was
+  /// created or cleared (whichever is latest).
+  func onFirstFrameRendered(controllerId controllerIdArg: Int64, completion: @escaping (Result<Void, FlutterError>) -> Void)
+  /// Called when the size of a VideoViewController's frame size changes.
+  func onFrameSizeChanged(controllerId controllerIdArg: Int64, width widthArg: Int64, height heightArg: Int64, completion: @escaping (Result<Void, FlutterError>) -> Void)
+}
+class RtcVideoViewControllerFlutterApi: RtcVideoViewControllerFlutterApiProtocol {
   private let binaryMessenger: FlutterBinaryMessenger
   init(binaryMessenger: FlutterBinaryMessenger){
     self.binaryMessenger = binaryMessenger
   }
+  /// Called when the VideoViewController with the given [controllerId] has rendered its first frame since it was
+  /// created or cleared (whichever is latest).
+  func onFirstFrameRendered(controllerId controllerIdArg: Int64, completion: @escaping (Result<Void, FlutterError>) -> Void) {
+    let channelName: String = "dev.flutter.pigeon.live_link_vcam.RtcVideoViewControllerFlutterApi.onFirstFrameRendered"
+    let channel = FlutterBasicMessageChannel(name: channelName, binaryMessenger: binaryMessenger)
+    channel.sendMessage([controllerIdArg] as [Any?]) { response in
+      guard let listResponse = response as? [Any?] else {
+        completion(.failure(createConnectionError(withChannelName:channelName)))
+        return
+      }
+      if (listResponse.count > 1) {
+        let code: String = listResponse[0] as! String
+        let message: String? = nilOrValue(listResponse[1])
+        let details: String? = nilOrValue(listResponse[2])
+        completion(.failure(FlutterError(code: code, message: message, details: details)));
+      } else {
+        completion(.success(Void()))
+      }
+    }
+  }
   /// Called when the size of a VideoViewController's frame size changes.
-  func onFrameSizeChanged(controllerId controllerIdArg: Int64, width widthArg: Int64, height heightArg: Int64, completion: @escaping () -> Void) {
-    let channel = FlutterBasicMessageChannel(name: "dev.flutter.pigeon.live_link_vcam.RtcVideoViewControllerFlutterApi.onFrameSizeChanged", binaryMessenger: binaryMessenger)
-    channel.sendMessage([controllerIdArg, widthArg, heightArg] as [Any?]) { _ in
-      completion()
+  func onFrameSizeChanged(controllerId controllerIdArg: Int64, width widthArg: Int64, height heightArg: Int64, completion: @escaping (Result<Void, FlutterError>) -> Void) {
+    let channelName: String = "dev.flutter.pigeon.live_link_vcam.RtcVideoViewControllerFlutterApi.onFrameSizeChanged"
+    let channel = FlutterBasicMessageChannel(name: channelName, binaryMessenger: binaryMessenger)
+    channel.sendMessage([controllerIdArg, widthArg, heightArg] as [Any?]) { response in
+      guard let listResponse = response as? [Any?] else {
+        completion(.failure(createConnectionError(withChannelName:channelName)))
+        return
+      }
+      if (listResponse.count > 1) {
+        let code: String = listResponse[0] as! String
+        let message: String? = nilOrValue(listResponse[1])
+        let details: String? = nilOrValue(listResponse[2])
+        completion(.failure(FlutterError(code: code, message: message, details: details)));
+      } else {
+        completion(.success(Void()))
+      }
     }
   }
 }
diff --git a/Engine/Extras/VirtualProduction/FlutterLiveLinkVCAM/IOS/Runner/WebRTC/Flutter/FlutterRtcDataChannel.swift b/Engine/Extras/VirtualProduction/FlutterLiveLinkVCAM/IOS/Runner/WebRTC/Flutter/FlutterRtcDataChannel.swift
--- a/Engine/Extras/VirtualProduction/FlutterLiveLinkVCAM/IOS/Runner/WebRTC/Flutter/FlutterRtcDataChannel.swift
+++ b/Engine/Extras/VirtualProduction/FlutterLiveLinkVCAM/IOS/Runner/WebRTC/Flutter/FlutterRtcDataChannel.swift
@@ -33,7 +33,7 @@ extension FlutterRtcDataChannel: RTCDataChannelDelegate {
       flutter.onStateChanged(
         dataChannelId: self.id,
         state: self.channel.readyState.toFlutter()
-      ) {}
+      ) { _ in }
     }
     
     if (self.channel.readyState == .closed) {
@@ -49,7 +49,7 @@ extension FlutterRtcDataChannel: RTCDataChannelDelegate {
           data: FlutterStandardTypedData(bytes: buffer.data),
           bIsBinary: buffer.isBinary
         )
-      ) {}
+      ) { _ in }
     }
   }
 }
diff --git a/Engine/Extras/VirtualProduction/FlutterLiveLinkVCAM/IOS/Runner/WebRTC/Flutter/FlutterRtcPeerConnection.swift b/Engine/Extras/VirtualProduction/FlutterLiveLinkVCAM/IOS/Runner/WebRTC/Flutter/FlutterRtcPeerConnection.swift
--- a/Engine/Extras/VirtualProduction/FlutterLiveLinkVCAM/IOS/Runner/WebRTC/Flutter/FlutterRtcPeerConnection.swift
+++ b/Engine/Extras/VirtualProduction/FlutterLiveLinkVCAM/IOS/Runner/WebRTC/Flutter/FlutterRtcPeerConnection.swift
@@ -177,13 +177,13 @@ extension FlutterRtcPeerConnection: RTCPeerConnectionDelegate {
     )
     
     api.callFlutter { flutter in
-      flutter.onIceCandidate(connectionId: self.id, candidate: data) {}
+      flutter.onIceCandidate(connectionId: self.id, candidate: data) { _ in }
     }
   }
   
   func peerConnection(_ peerConnection: RTCPeerConnection, didChange newState: RTCPeerConnectionState) {
     api.callFlutter { flutter in
-      flutter.onStateChanged(connectionId: self.id, state: newState.toFlutter()) {}
+      flutter.onStateChanged(connectionId: self.id, state: newState.toFlutter()) { _ in }
     }
   }
   
@@ -202,7 +202,7 @@ extension FlutterRtcPeerConnection: RTCPeerConnectionDelegate {
           trackId: trackId,
           kind: WebRtcUtils.parseTrackKind(string: track!.kind)
         )
-      ) {}
+      ) { _ in }
     }
   }
   
@@ -210,7 +210,7 @@ extension FlutterRtcPeerConnection: RTCPeerConnectionDelegate {
     let dataChannel: FlutterRtcDataChannel = FlutterRtcDataChannelApi.instance!.register(dataChannel: dataChannel)
     
     api.callFlutter { flutter in
-      flutter.onDataChannel(connectionId: self.id, dataChannelId: dataChannel.id) {}
+      flutter.onDataChannel(connectionId: self.id, dataChannelId: dataChannel.id) { _ in }
     }
   }
   
diff --git a/Engine/Extras/VirtualProduction/FlutterLiveLinkVCAM/IOS/Runner/WebRTC/Flutter/FlutterRtcVideoViewController.swift b/Engine/Extras/VirtualProduction/FlutterLiveLinkVCAM/IOS/Runner/WebRTC/Flutter/FlutterRtcVideoViewController.swift
--- a/Engine/Extras/VirtualProduction/FlutterLiveLinkVCAM/IOS/Runner/WebRTC/Flutter/FlutterRtcVideoViewController.swift
+++ b/Engine/Extras/VirtualProduction/FlutterLiveLinkVCAM/IOS/Runner/WebRTC/Flutter/FlutterRtcVideoViewController.swift
@@ -1,6 +1,9 @@
 // Copyright Epic Games, Inc. All Rights Reserved.
 
+import Flutter
 import Foundation
+import MetalKit
+import simd
 import WebRTC
 
 protocol FlutterRtcVideoViewControllerListener: AnyObject {
@@ -12,30 +15,135 @@ protocol FlutterRtcVideoViewControllerListener: AnyObject {
   func videoViewController(didChangeVideoTrack track: RTCVideoTrack?)
 }
 
-/// Controller for a video view which communicates with Flutter.
+/// Vertex data used by the render pipeline.
+struct Vertex {
+  /// The position of the vertex in the scene.
+  var position: SIMD2<Float>
+  
+  /// The texture coordinate for this vertex.
+  var texCoord: SIMD2<Float>
+}
+
+/// Controller for a WebRTC video view which renders a video stream to a shared texture shared with Flutter and
+/// communicates related messages.
 class FlutterRtcVideoViewController: IdObject {
+  /// The ID of the output texture bound to the Flutter engine.
+  public private(set) var textureId: Int64 = -1
+  
+  /// The Metal device used to render WebRTC frame data.
+  private let device: MTLDevice
+  
+  /// Cache containing Metal textures.
+  private var textureCache: CVMetalTextureCache?
+  
+  /// Metal texture-backed buffer for the renderTexture.
+  private var renderTextureBuffer: CVMetalTexture?
+  
+  /// Handle of the Metal texture registered for display in Flutter.
+  private var renderTexture: MTLTexture?
+  
+  /// Texture containing the Y (luma) channel of the WebRTC video stream.
+  private var yTexture: MTLTexture?
+  
+  /// Texture containing the U (chroma) channel of the WebRTC video stream.
+  private var uTexture: MTLTexture?
+  
+  /// Texture containing the V (chroma) channel of the WebRTC video stream.
+  private var vTexture: MTLTexture?
+  
+  /// Descriptor of yTexture.
+  private var yDescriptor: MTLTextureDescriptor?
+  
+  /// Descriptor of uTexture/vTexture.
+  private var uvDescriptor: MTLTextureDescriptor?
+  
+  /// Pipeline state containing our configuration for Metal rendering.
+  var pipelineState: MTLRenderPipelineState
+  
+  /// Queue of commands to run on the GPU.
+  var commandQueue: MTLCommandQueue
+  
+  /// Descriptor for the rendering pass that converts I420 to BGRA data.
+  var convertPassDescriptor = MTLRenderPassDescriptor()
+  
   /// The API used to send messages to Flutter.
   private let api: FlutterRtcVideoViewControllerApi
   
+  /// Vertices for the quad used to render the converted video stream.
+  private let vertices = [
+    Vertex(position: [-1.0, -1.0], texCoord: [0.0, 1.0]),
+    Vertex(position: [ 1.0, -1.0], texCoord: [1.0, 1.0]),
+    Vertex(position: [-1.0,  1.0], texCoord: [0.0, 0.0]),
+    Vertex(position: [ 1.0,  1.0], texCoord: [1.0, 0.0]),
+  ]
+  
   /// The current video track to be displayed.
-  private var _track: RTCVideoTrack?
+  private var track: RTCVideoTrack?
   
-  /// Listeners for changes to the changes in the controller's state.
-  private var listeners: [FlutterRtcVideoViewControllerListener] = []
+  /// Whether this has reported its first frame since being created/cleared.
+  private var bHasReportedFirstFrame: Bool = false
   
-  /// The current video track to be displayed.
-  var track: RTCVideoTrack? {
-    get {
-      return _track
-    }
-  }
+  /// The Core Video pixel buffer backing the Flutter texture for the video stream.
+  private var pixelBuffer: Unmanaged<CVPixelBuffer>?
+  
+  /// The width of the last received frame.
+  private var lastFrameWidth: Int?
+  
+  /// The height of the last received frame.
+  private var lastFrameHeight: Int?
   
   init(api: FlutterRtcVideoViewControllerApi) {
     self.api = api
+    
+    guard let device = MTLCreateSystemDefaultDevice()
+    else {
+        fatalError("Failed to create default Metal device for WebRTC stream")
+    }
+    self.device = device
+    
+    guard CVMetalTextureCacheCreate(
+      kCFAllocatorDefault,
+      nil,
+      device,
+      nil,
+      &textureCache
+    ) == kCVReturnSuccess
+    else {
+      fatalError("Failed to create texture cache for WebRTC stream")
+    }
+    
+    let shaderLib = device.makeDefaultLibrary()
+    
+    let pipelineDescriptor = MTLRenderPipelineDescriptor()
+    pipelineDescriptor.label = "FlutterRtcPipeline"
+    pipelineDescriptor.colorAttachments[0].pixelFormat = .bgra8Unorm
+    pipelineDescriptor.vertexFunction = shaderLib?.makeFunction(name: "i420ToBgraVertex")
+    pipelineDescriptor.fragmentFunction = shaderLib?.makeFunction(name: "i420ToBgraFragment")
+    
+    convertPassDescriptor.colorAttachments[0].loadAction = .clear
+    convertPassDescriptor.colorAttachments[0].clearColor = MTLClearColor(red: 0, green: 0, blue: 0, alpha: 0)
+    
+    guard let pipelineState = try? device.makeRenderPipelineState(descriptor: pipelineDescriptor)
+    else {
+      fatalError("Failed to create render pipeline for WebRTC stream")
+    }
+    self.pipelineState = pipelineState
+    
+    guard let commandQueue = device.makeCommandQueue()
+    else {
+      fatalError("Failed to create command queue for WebRTC stream")
+    }
+    self.commandQueue = commandQueue
+    
+    super.init()
+    
+    textureId = api.textureRegistry.register(self)
   }
   
   override func dispose() {
-    _track?.remove(self)
+    api.textureRegistry.unregisterTexture(textureId)
+    track?.remove(self)
+    pixelBuffer = nil
   }
   
   /**
@@ -54,39 +162,211 @@ class FlutterRtcVideoViewController: IdObject {
       }
     }
     
-    _track = newVideoTrack!
-    _track!.add(self)
+    track = newVideoTrack!
+    track!.add(self)
     
-    listeners.forEach { listener in
-      listener.videoViewController(didChangeVideoTrack: _track)
-    }
+    clear()
   }
   
-  /// Listen for state changes to this controller.
-  func addListener(listener: FlutterRtcVideoViewControllerListener) {
-    listeners.append(listener)
+  /// Clear the rendered video stream texture.
+  func clear() {
+    bHasReportedFirstFrame = false
+    
+    if (pixelBuffer == nil || renderTexture == nil || renderTexture!.bufferBytesPerRow <= 0) {
+      return
+    }
+    
+    let clearBytes = [Float](
+      repeating: 0.0,
+      count: renderTexture!.allocatedSize
+    )
+
+    renderTexture!.replace(
+      region: MTLRegionMake2D(0, 0, renderTexture!.width, renderTexture!.height),
+      mipmapLevel: 0,
+      withBytes: clearBytes,
+      bytesPerRow: renderTexture!.width * renderTexture!.bufferBytesPerRow
+    )
   }
   
-  /// Stop listening for state changes to this controller.
-  func removeListener(listener: FlutterRtcVideoViewControllerListener) {
-    guard let index: Int = listeners.firstIndex(where: { $0 === listener }) else {
-      return
+  /// Render the stored I420 data to the BGRA texture shared with Flutter.
+  private func convertFrame() {
+    guard let commandBuffer = commandQueue.makeCommandBuffer()
+    else {
+      fatalError("Failed to create command buffer for WebRTC stream")
     }
+    commandBuffer.label = "FlutterRtcCommandBuffer"
     
-    listeners.remove(at: index)
+    guard let encoder = commandBuffer.makeRenderCommandEncoder(descriptor: convertPassDescriptor)
+    else {
+      fatalError("Failed to create command encoder for WebRTC stream")
+    }
+    encoder.label = "FlutterRtcEncoder"
+    
+    encoder.pushDebugGroup("FlutterRtcRenderFrame")
+    encoder.setRenderPipelineState(pipelineState)
+    
+    encoder.setFragmentTexture(yTexture, index: 0)
+    encoder.setFragmentTexture(uTexture, index: 1)
+    encoder.setFragmentTexture(vTexture, index: 2)
+    
+    // Draw the quad
+    encoder.setVertexBytes(
+      vertices,
+      length: MemoryLayout<Vertex>.stride * vertices.count,
+      index: 0
+    )
+    encoder.drawPrimitives(type: .triangleStrip, vertexStart: 0, vertexCount: 4, instanceCount: 1)
+    
+    encoder.popDebugGroup()
+    encoder.endEncoding()
+    
+    commandBuffer.commit()
   }
 }
 
+//MARK: RTCVideoRenderer
 extension FlutterRtcVideoViewController: RTCVideoRenderer {
   func setSize(_ size: CGSize) {
+    let width = Int(size.width)
+    let height = Int(size.height)
+    
+    if (width <= 0 || height <= 0) {
+      return
+    }
+    
+    // Create a surface for inter-process sharing
+    guard let surface = IOSurfaceCreate([
+      kIOSurfaceWidth: width,
+      kIOSurfaceHeight: height,
+      kIOSurfaceBytesPerElement: 4,
+      kIOSurfacePixelFormat: kCVPixelFormatType_32BGRA
+    ] as CFDictionary)
+    else {
+      fatalError("Failed to create surface for WebRTC stream")
+    }
+    
+    // Create a new pixel buffer of the correct size. Any previous buffer will be released automatically
+    guard CVPixelBufferCreateWithIOSurface(
+      kCFAllocatorDefault,
+      surface,
+      [ kCVPixelBufferCGImageCompatibilityKey: true, kCVPixelBufferCGBitmapContextCompatibilityKey: true ] as CFDictionary,
+      &pixelBuffer
+    ) == 0
+    else {
+      fatalError("Failed to create pixel buffer for WebRTC stream")
+    }
+    
+    guard CVMetalTextureCacheCreateTextureFromImage(
+      kCFAllocatorDefault,
+      textureCache!, pixelBuffer!.takeUnretainedValue(),
+      nil, .bgra8Unorm,
+      width, height, 0,
+      &renderTextureBuffer
+    ) == kCVReturnSuccess
+    else {
+      fatalError("Failed to create Metal texture buffer for WebRTC stream")
+    }
+    
+    renderTexture = CVMetalTextureGetTexture(renderTextureBuffer!)
+    convertPassDescriptor.colorAttachments[0].texture = renderTexture
+    
+    api.textureRegistry.textureFrameAvailable(textureId)
+    
     api.callFlutter { flutter in
       flutter.onFrameSizeChanged(
         controllerId: self.id,
         width: Int64(size.width),
         height: Int64(size.height)
-      ) {}
+      ) { _ in }
     }
   }
   
-  func renderFrame(_ frame: RTCVideoFrame?) {}
+  func renderFrame(_ frame: RTCVideoFrame?) {
+    if (frame == nil || renderTexture == nil) {
+      return
+    }
+    
+    let width = Int(frame!.width)
+    let height = Int(frame!.height)
+    let chromaWidth = width / 2
+    let chromaHeight = height / 2
+    
+    let hasSizeChanged = width != lastFrameWidth || height != lastFrameHeight
+    let yuvBuffer: RTCI420BufferProtocol = frame!.buffer.toI420()
+    
+    // Create/update luma (Y) texture from frame data
+    if (hasSizeChanged || yTexture == nil) {
+      yDescriptor = MTLTextureDescriptor.texture2DDescriptor(
+        pixelFormat: .r8Unorm,
+        width: width,
+        height: height,
+        mipmapped: false
+      )
+      yDescriptor?.usage = .shaderRead
+      
+      yTexture = device.makeTexture(descriptor: yDescriptor!)
+    }
+    
+    yTexture!.replace(
+      region: MTLRegionMake2D(0, 0, width, height),
+      mipmapLevel: 0,
+      withBytes: yuvBuffer.dataY,
+      bytesPerRow: Int(yuvBuffer.strideY)
+    )
+    
+    // Create/update chroma (U/V) texture from frame data
+    if (hasSizeChanged || uTexture == nil || vTexture == nil) {
+      uvDescriptor = MTLTextureDescriptor.texture2DDescriptor(
+        pixelFormat: .r8Unorm,
+        width: width / 2,
+        height: height / 2,
+        mipmapped: false
+      )
+      uvDescriptor?.usage = .shaderRead
+      
+      uTexture = device.makeTexture(descriptor: uvDescriptor!)
+      vTexture = device.makeTexture(descriptor: uvDescriptor!)
+    }
+    
+    uTexture!.replace(
+      region: MTLRegionMake2D(0, 0, chromaWidth, chromaHeight),
+      mipmapLevel: 0,
+      withBytes: yuvBuffer.dataU,
+      bytesPerRow: Int(yuvBuffer.strideU)
+    )
+    
+    vTexture!.replace(
+      region: MTLRegionMake2D(0, 0, chromaWidth, chromaHeight),
+      mipmapLevel: 0,
+      withBytes: yuvBuffer.dataV,
+      bytesPerRow: Int(yuvBuffer.strideV)
+    )
+    
+    lastFrameWidth = width
+    lastFrameHeight = height
+    
+    convertFrame()
+    
+    if (!bHasReportedFirstFrame) {
+      api.callFlutter { flutter in
+        flutter.onFirstFrameRendered(controllerId: self.id) { _ in }
+      }
+      
+      bHasReportedFirstFrame = true
+    }
+  }
+}
+
+//MARK: FlutterTexture
+extension FlutterRtcVideoViewController: FlutterTexture {
+  func copyPixelBuffer() -> Unmanaged<CVPixelBuffer>? {
+    guard let buffer = pixelBuffer?.takeUnretainedValue()
+    else {
+      return nil
+    }
+    
+    // Retain since Flutter releases the texture when finished
+    return Unmanaged.passRetained(buffer)
+  }
 }
