[LiveLinkVCAM] Tentacle device details menu

diff --git a/Engine/Extras/VirtualProduction/FlutterLiveLinkVCAM/pubspec.yaml b/Engine/Extras/VirtualProduction/FlutterLiveLinkVCAM/pubspec.yaml
--- a/Engine/Extras/VirtualProduction/FlutterLiveLinkVCAM/pubspec.yaml
+++ b/Engine/Extras/VirtualProduction/FlutterLiveLinkVCAM/pubspec.yaml
@@ -57,6 +57,7 @@ flutter:
 
   assets:
     - assets/icons/
+    - assets/icons/tentacle/
 
   fonts:
     - family: Inter
