[UnrealStage] Remember camera angle between runs. Add camera control dropdown

diff --git a/Engine/Extras/VirtualProduction/EpicStageApp/lib/main.dart b/Engine/Extras/VirtualProduction/EpicStageApp/lib/main.dart
--- a/Engine/Extras/VirtualProduction/EpicStageApp/lib/main.dart
+++ b/Engine/Extras/VirtualProduction/EpicStageApp/lib/main.dart
@@ -166,7 +166,10 @@ class _EpicStageAppState extends State<EpicStageApp> with TickerProviderStateMix
         Provider<PreferencesBundle>(create: (_) => _preferenceBundle),
         Provider<ConnectionSettings>(create: (_) => ConnectionSettings(_preferenceBundle)),
         Provider<SelectedActorSettings>(create: (_) => SelectedActorSettings(_preferenceBundle)),
-        Provider<StageMapSettings>(create: (_) => StageMapSettings(_preferenceBundle)),
+        Provider<StageMapSettings>(
+          create: (context) => StageMapSettings(_preferenceBundle, context),
+          dispose: (_, value) => value.dispose(),
+        ),
         Provider<RecentActorSettings>(create: (_) => RecentActorSettings(_preferenceBundle)),
         Provider<MainScreenSettings>(create: (_) => MainScreenSettings(_preferenceBundle)),
         Provider<DeltaWidgetSettings>(create: (_) => DeltaWidgetSettings(_preferenceBundle)),
diff --git a/Engine/Extras/VirtualProduction/EpicStageApp/lib/models/preview_render_manager.dart b/Engine/Extras/VirtualProduction/EpicStageApp/lib/models/preview_render_manager.dart
--- a/Engine/Extras/VirtualProduction/EpicStageApp/lib/models/preview_render_manager.dart
+++ b/Engine/Extras/VirtualProduction/EpicStageApp/lib/models/preview_render_manager.dart
@@ -128,10 +128,11 @@ class PreviewRenderManager with WidgetsBindingObserver {
 
     _settingSubscriptions.addAll([
       _selectedActorSettings.displayClusterRootPath.listen(_onRootActorPathChanged),
-      _stageMapSettings.projectionMode.listen(_onProjectionModeChanged),
-      _stageMapSettings.cameraAngle.listen(_onCameraAngleChanged),
     ]);
 
+    _stageMapSettings.projectionMode.addListener(_onProjectionModeChanged);
+    _stageMapSettings.cameraAngle.addListener(_onCameraAngleChanged);
+
     WidgetsBinding.instance.addObserver(this);
   }
 
@@ -272,6 +273,9 @@ class PreviewRenderManager with WidgetsBindingObserver {
     for (final StreamSubscription subscription in _settingSubscriptions) {
       subscription.cancel();
     }
+
+    _stageMapSettings.projectionMode.removeListener(_onProjectionModeChanged);
+    _stageMapSettings.cameraAngle.removeListener(_onCameraAngleChanged);
   }
 
   /// Start listening for messages from the engine connection manager.
@@ -371,7 +375,7 @@ class PreviewRenderManager with WidgetsBindingObserver {
             },
             'Rotation': _getCameraRotationSettings(),
             'IncludeActorPositions': true,
-            'ProjectionType': _getProjectionModeTypeName(_stageMapSettings.projectionMode.getValue()),
+            'ProjectionType': _getProjectionModeTypeName(_stageMapSettings.projectionMode.value),
             'FOV': _getFOV(),
           },
         });
@@ -563,15 +567,15 @@ class PreviewRenderManager with WidgetsBindingObserver {
   }
 
   /// Update the preview renderer camera's angle.
-  void _onCameraAngleChanged(vec.Vector2 cameraAngle) {
+  void _onCameraAngleChanged() {
     _tickRendererSettingsChanges['Rotation'] = _getCameraRotationSettings();
   }
 
   /// Update the projection mode and affected settings.
-  void _onProjectionModeChanged(ProjectionMode projectionMode) {
-    _tickRendererSettingsChanges['ProjectionType'] = _getProjectionModeTypeName(projectionMode);
+  void _onProjectionModeChanged() {
+    _tickRendererSettingsChanges['ProjectionType'] = _getProjectionModeTypeName(_stageMapSettings.projectionMode.value);
     _tickRendererSettingsChanges['FOV'] = _getFOV();
-    _onCameraAngleChanged(_stageMapSettings.cameraAngle.getValue());
+    _onCameraAngleChanged();
   }
 
   /// Called when the root actors in the engine have changed.
@@ -618,7 +622,7 @@ class PreviewRenderManager with WidgetsBindingObserver {
 
   /// Get the field of view based on the current projection mode.
   double _getFOV() {
-    switch (_stageMapSettings.projectionMode.getValue()) {
+    switch (_stageMapSettings.projectionMode.value) {
       case ProjectionMode.uv:
         return 45.0;
 
@@ -633,8 +637,8 @@ class PreviewRenderManager with WidgetsBindingObserver {
 
   /// Get the JSON data representing the camera's rotation settings based on the projection mode and user settings.
   dynamic _getCameraRotationSettings() {
-    final bool bIsUvMode = _stageMapSettings.projectionMode.getValue() == ProjectionMode.uv;
-    final vec.Vector2 cameraAngle = _stageMapSettings.cameraAngle.getValue();
+    final bool bIsUvMode = _stageMapSettings.projectionMode.value == ProjectionMode.uv;
+    final vec.Vector2 cameraAngle = _stageMapSettings.cameraAngle.value;
 
     return {
       'Yaw': bIsUvMode ? 0 : cameraAngle.x,
diff --git a/Engine/Extras/VirtualProduction/EpicStageApp/lib/models/settings/stage_map_settings.dart b/Engine/Extras/VirtualProduction/EpicStageApp/lib/models/settings/stage_map_settings.dart
--- a/Engine/Extras/VirtualProduction/EpicStageApp/lib/models/settings/stage_map_settings.dart
+++ b/Engine/Extras/VirtualProduction/EpicStageApp/lib/models/settings/stage_map_settings.dart
@@ -1,11 +1,17 @@
 // Copyright Epic Games, Inc. All Rights Reserved.
 
-import 'dart:ui';
+import 'dart:async';
 
 import 'package:epic_common/preferences.dart';
+import 'package:epic_common/utilities/json.dart';
+import 'package:flutter/cupertino.dart';
+import 'package:flutter/foundation.dart';
+import 'package:provider/provider.dart';
 import 'package:streaming_shared_preferences/streaming_shared_preferences.dart';
 import 'package:vector_math/vector_math_64.dart' as vec;
 
+import 'selected_actor_settings.dart';
+
 enum ProjectionMode {
   azimuthal,
   orthographic,
@@ -15,19 +21,28 @@ enum ProjectionMode {
 
 /// Holds user settings used to render and control the stage map.
 class StageMapSettings {
-  StageMapSettings(PreferencesBundle preferences)
+  StageMapSettings(PreferencesBundle preferences, BuildContext context)
       : bIsInTransformMode = preferences.persistent.getBool(
           'stageMap.bIsInMultiSelectMode',
           defaultValue: false,
         ),
-        projectionMode = preferences.persistent.getEnum<ProjectionMode>(
-          'stageMap.projectionMode',
-          defaultValue: ProjectionMode.azimuthal,
-          enumValues: ProjectionMode.values,
-        ),
-        cameraAngle = preferences.transient.get(
-          'stageMap.cameraAngle',
-          defaultValue: vec.Vector2(0, 90),
+        perRootActorPreferences = preferences.persistent.getCustomValue(
+          'stageMap.perRootActor',
+          defaultValue: {},
+          adapter: JsonAdapter(
+            serializer: (Map<String, _PerRootActorPreferenceGroup> map) => map.map(
+              (String key, _PerRootActorPreferenceGroup value) => MapEntry(key, value.toJson()),
+            ),
+            deserializer: (json) {
+              if (!(json is Map<String, dynamic>)) {
+                return {};
+              }
+
+              return json.map(
+                (key, value) => MapEntry(key, _PerRootActorPreferenceGroup.fromJson(value)),
+              );
+            },
+          ),
         ),
         focalPoint = preferences.transient.get(
           'stageMap.focalPoint',
@@ -36,7 +51,21 @@ class StageMapSettings {
         zoomLevel = preferences.transient.get(
           'stageMap.zoomLevel',
           defaultValue: 1.0,
-        );
+        ) {
+    projectionMode = PerRootActorPreference<ProjectionMode>(
+      context: context,
+      perRootActorPreferences: perRootActorPreferences,
+      getter: (group) => group.projectionMode,
+      setter: (group, value) => group.copyWith(projectionMode: value),
+    );
+
+    cameraAngle = PerRootActorPreference<vec.Vector2>(
+      context: context,
+      perRootActorPreferences: perRootActorPreferences,
+      getter: (group) => group.cameraAngle.clone(), // Clone this to avoid mutating the cached vector
+      setter: (group, value) => group.copyWith(cameraAngle: value),
+    );
+  }
 
   /// If true, the map is in transform mode and the user can pinch/pan to adjust lightcard scale/rotation/location.
   /// If false, the map is in map mode, and the user can pinch/pan to move the map or long-press pins to drag them
@@ -44,14 +73,136 @@ class StageMapSettings {
   final Preference<bool> bIsInTransformMode;
 
   /// The camera's current projection mode for stage actor locations.
-  final Preference<ProjectionMode> projectionMode;
+  late final PerRootActorPreference<ProjectionMode> projectionMode;
 
   /// The camera's current rotation setting in degrees, where X is yaw and Y is pitch.
-  final TransientPreference<vec.Vector2> cameraAngle;
+  late final PerRootActorPreference<vec.Vector2> cameraAngle;
+
+  /// Map from root actor path to preferences associated only with that root actor.
+  final Preference<Map<String, _PerRootActorPreferenceGroup>> perRootActorPreferences;
 
   /// The location on the map in center of the stage map view, assuming a 1:1 aspect ratio.
   final TransientPreference<Offset> focalPoint;
 
   /// The zoom level of the stage map view.
   final TransientPreference<double> zoomLevel;
+
+  /// Dispose of any resources held by this.
+  void dispose() {
+    projectionMode.dispose();
+    cameraAngle.dispose();
+  }
+}
+
+/// Contains/serializes/deserializes all preferences that depend on the currently selected root actor.
+class _PerRootActorPreferenceGroup {
+  _PerRootActorPreferenceGroup({
+    ProjectionMode? projectionMode,
+    vec.Vector2? cameraAngle,
+  })  : this.projectionMode = projectionMode ?? ProjectionMode.azimuthal,
+        this.cameraAngle = cameraAngle ?? vec.Vector2(0, 90);
+
+  /// Create a preference group from JSON data.
+  _PerRootActorPreferenceGroup.fromJson(Map<String, dynamic> json)
+      : this(
+          projectionMode: jsonToEnumValue(json['ProjectionMode'], ProjectionMode.values),
+          cameraAngle: jsonToVector2(json['CameraAngle']),
+        );
+
+  /// Convert a preference group to JSON data.
+  Map<String, dynamic> toJson() => {
+        'ProjectionMode': enumToJsonValue<ProjectionMode>(projectionMode),
+        'CameraAngle': cameraAngle.toJson(),
+      };
+
+  /// Copy the group, changing one or more of its values.
+  _PerRootActorPreferenceGroup copyWith({
+    ProjectionMode? projectionMode,
+    vec.Vector2? cameraAngle,
+  }) =>
+      _PerRootActorPreferenceGroup(
+        projectionMode: projectionMode ?? this.projectionMode,
+        cameraAngle: cameraAngle ?? this.cameraAngle,
+      );
+
+  /// The camera's current projection mode for stage actor locations.
+  final ProjectionMode projectionMode;
+
+  /// The camera's current rotation setting in degrees, where X is yaw and Y is pitch.
+  final vec.Vector2 cameraAngle;
+}
+
+/// Provides access to a single preference that depends on the currently selected root actor.
+class PerRootActorPreference<T> extends ChangeNotifier implements ValueListenable<T> {
+  PerRootActorPreference({
+    required BuildContext context,
+    required this.perRootActorPreferences,
+    required this.getter,
+    required this.setter,
+  })  : _displayClusterRootPath = Provider.of<SelectedActorSettings>(context, listen: false).displayClusterRootPath,
+        _group = _PerRootActorPreferenceGroup() {
+    _subscriptions.addAll([
+      _displayClusterRootPath.listen((value) => _retrieveAndNotify()),
+      perRootActorPreferences.listen((value) => _retrieveAndNotify()),
+    ]);
+
+    _retrieveAndNotify();
+  }
+
+  /// Preference containing the map from root actor paths to associated settings.
+  final Preference<Map<String, _PerRootActorPreferenceGroup>> perRootActorPreferences;
+
+  /// Function which retrieves the preference from the given preference [group].
+  final T Function(_PerRootActorPreferenceGroup group) getter;
+
+  /// Function which returns a new preference group by taking an existing [group] and applying the preference's new
+  /// [value] to it.
+  final _PerRootActorPreferenceGroup Function(_PerRootActorPreferenceGroup group, T value) setter;
+
+  /// Preference containing the path of the selected root actor.
+  final Preference<String> _displayClusterRootPath;
+
+  /// Subscriptions to clean up when this is disposed.
+  final List<StreamSubscription> _subscriptions = [];
+
+  /// Cached group containing the preference.
+  _PerRootActorPreferenceGroup _group;
+
+  /// Clean up any references held by this.
+  @override
+  void dispose() {
+    _subscriptions.forEach((subscription) => subscription.cancel());
+
+    super.dispose();
+  }
+
+  /// The current value of the preference.
+  @override
+  T get value => getter(_group);
+
+  set value(T newValue) {
+    if (value == newValue) {
+      return;
+    }
+
+    final String path = _displayClusterRootPath.getValue();
+    final Map<String, _PerRootActorPreferenceGroup> prefMap = perRootActorPreferences.getValue();
+    final _PerRootActorPreferenceGroup group = prefMap[path] ?? _PerRootActorPreferenceGroup();
+
+    prefMap[path] = setter(group, newValue);
+
+    perRootActorPreferences.setValue(prefMap); // This will in turn notify us, causing us to update the cached value
+  }
+
+  /// Update the cached value of the preference based on the stored data and notify listeners.
+  void _retrieveAndNotify() {
+    final String path = _displayClusterRootPath.getValue();
+
+    // Cache the group rather than using the getter here in case the getter needs to run additional operations at
+    // the time of access. For example, we may want each access to create a new copy rather than returning the same
+    // cached copy.
+    _group = perRootActorPreferences.getValue()[path] ?? _PerRootActorPreferenceGroup();
+
+    notifyListeners();
+  }
 }
diff --git a/Engine/Extras/VirtualProduction/EpicStageApp/lib/widgets/elements/lightcard_map.dart b/Engine/Extras/VirtualProduction/EpicStageApp/lib/widgets/elements/lightcard_map.dart
--- a/Engine/Extras/VirtualProduction/EpicStageApp/lib/widgets/elements/lightcard_map.dart
+++ b/Engine/Extras/VirtualProduction/EpicStageApp/lib/widgets/elements/lightcard_map.dart
@@ -46,6 +46,9 @@ enum _GestureState {
   cameraTrackingAndZooming,
 }
 
+/// Radius for decoration of dropdowns in the control panel.
+const Radius _panelDropdownRadius = Radius.circular(20);
+
 /// A map displaying all stage actors pins on an azimuthal projection map of the scene.
 class StageMap extends StatefulWidget {
   const StageMap({
@@ -343,22 +346,24 @@ class StageMapState extends State<StageMap> with PreviewRenderConsumer, GuardedR
                               spacing: 16,
                               children: [
                                 const _StageMapControlModeToggle(),
+
+                                // View mode
                                 ConstrainedBox(
                                   constraints: const BoxConstraints(minWidth: 128),
                                   child: Tooltip(
                                     message: AppLocalizations.of(context)!.stageMapProjectionMode,
-                                    child: PreferenceBuilder<ProjectionMode>(
-                                      preference: stageMapSettings.projectionMode,
-                                      builder: (context, projectionMode) {
+                                    child: ValueListenableBuilder<ProjectionMode>(
+                                      valueListenable: stageMapSettings.projectionMode,
+                                      builder: (context, projectionMode, _) {
                                         // Projection mode control
                                         return DropdownSelector<ProjectionMode>(
-                                          borderRadius: const Radius.circular(20),
+                                          borderRadius: _panelDropdownRadius,
                                           value: projectionMode,
                                           items: ProjectionMode.values,
                                           makeItemName: (ProjectionMode value) => _getProjectionModeName(value),
                                           onChanged: (ProjectionMode? value) {
                                             if (value != null) {
-                                              stageMapSettings.projectionMode.setValue(value);
+                                              stageMapSettings.projectionMode.value = value;
                                             }
                                           },
                                         );
@@ -366,6 +371,9 @@ class StageMapState extends State<StageMap> with PreviewRenderConsumer, GuardedR
                                     ),
                                   ),
                                 ),
+
+                                // View controls
+                                const _StageMapViewDropdownButton(),
                               ],
                             ),
                           ),
@@ -918,7 +926,7 @@ class StageMapState extends State<StageMap> with PreviewRenderConsumer, GuardedR
       } else if (details.pointerCount >= 2) {
         _onTransformScaleStart();
       }
-    } else if (details.pointerCount > 1 || stageMapSettings.projectionMode.getValue() == ProjectionMode.uv) {
+    } else if (details.pointerCount > 1 || stageMapSettings.projectionMode.value == ProjectionMode.uv) {
       // Start a camera track and zoom if we're using a scale gesture or just dragging in UV mode (since it can't
       // tumble)
       _onCameraTrackAndZoomStart(details);
@@ -941,7 +949,7 @@ class StageMapState extends State<StageMap> with PreviewRenderConsumer, GuardedR
         default:
           break;
       }
-    } else if (details.pointerCount == 1 && stageMapSettings.projectionMode.getValue() != ProjectionMode.uv) {
+    } else if (details.pointerCount == 1 && stageMapSettings.projectionMode.value != ProjectionMode.uv) {
       // We don't tumble in UV mode, so skip to track and zoom in that case
       _onCameraTumbleUpdate(details);
     } else {
@@ -999,11 +1007,11 @@ class StageMapState extends State<StageMap> with PreviewRenderConsumer, GuardedR
 
     final stageMapSettings = Provider.of<StageMapSettings>(context, listen: false);
 
-    vec.Vector2 cameraAngle = stageMapSettings.cameraAngle.getValue();
+    vec.Vector2 cameraAngle = stageMapSettings.cameraAngle.value;
     cameraAngle.x += details.focalPointDelta.dx * _tumbleSensitivity.dx;
     cameraAngle.y = (cameraAngle.y - details.focalPointDelta.dy * _tumbleSensitivity.dy).clamp(-90.0, 90.0);
 
-    stageMapSettings.cameraAngle.setValue(cameraAngle);
+    stageMapSettings.cameraAngle.value = cameraAngle;
   }
 
   /// Called when the user starts a scale gesture that should track and zoom the camera.
@@ -1963,3 +1971,70 @@ class _StageMapControlModeToggleButton extends StatelessWidget {
     );
   }
 }
+
+/// Button and dropdown menu controlling the camera view in the stage map.
+class _StageMapViewDropdownButton extends StatelessWidget {
+  const _StageMapViewDropdownButton();
+
+  @override
+  Widget build(BuildContext context) {
+    return ModalDropdownButton(
+      buttonBuilder: (context, state) => DropdownSelectorStyleButton(
+        state: state,
+        label: AppLocalizations.of(context)!.stageMapViewMenuLabel,
+        borderRadius: _panelDropdownRadius,
+      ),
+      menuBuilder: (context, originTabBuilder) => DropDownListMenu(
+        originTabBuilder: originTabBuilder,
+        children: [
+          const SizedBox(height: 16),
+          ListMenuHeader(AppLocalizations.of(context)!.stageMapViewMenuSectionOrientation),
+          ListMenuSimpleItem(
+            title: AppLocalizations.of(context)!.stageMapViewMenuOrientationTop,
+            iconPath: 'assets/images/icons/view_top.svg',
+            onTap: () => _setCameraAngle(context, vec.Vector2(0, 90)),
+          ),
+          ListMenuSimpleItem(
+            title: AppLocalizations.of(context)!.stageMapViewMenuOrientationBottom,
+            iconPath: 'assets/images/icons/view_bottom.svg',
+            onTap: () => _setCameraAngle(context, vec.Vector2(0, -90)),
+          ),
+          ListMenuSimpleItem(
+            title: AppLocalizations.of(context)!.stageMapViewMenuOrientationLeft,
+            iconPath: 'assets/images/icons/view_left.svg',
+            onTap: () => _setCameraAngle(context, vec.Vector2(-90, 0)),
+          ),
+          ListMenuSimpleItem(
+            title: AppLocalizations.of(context)!.stageMapViewMenuOrientationRight,
+            iconPath: 'assets/images/icons/view_right.svg',
+            onTap: () => _setCameraAngle(context, vec.Vector2(90, 0)),
+          ),
+          ListMenuSimpleItem(
+            title: AppLocalizations.of(context)!.stageMapViewMenuOrientationFront,
+            iconPath: 'assets/images/icons/view_front.svg',
+            onTap: () => _setCameraAngle(context, vec.Vector2(0, 0)),
+          ),
+          ListMenuSimpleItem(
+            title: AppLocalizations.of(context)!.stageMapViewMenuOrientationBack,
+            iconPath: 'assets/images/icons/view_back.svg',
+            onTap: () => _setCameraAngle(context, vec.Vector2(180, 0)),
+          ),
+          ListMenuHeader(AppLocalizations.of(context)!.stageMapViewMenuSectionOptions),
+          ListMenuSimpleItem(
+            title: AppLocalizations.of(context)!.stageMapViewMenuOptionsReset,
+            iconPath: 'packages/epic_common/assets/icons/reset.svg',
+            onTap: () => _setCameraAngle(context, vec.Vector2(0, 90)),
+          ),
+        ],
+      ),
+    );
+  }
+
+  /// Given the build [context], set the stage map camera [angle] and close the menu.
+  void _setCameraAngle(BuildContext context, vec.Vector2 angle) {
+    final stageMapSettings = Provider.of<StageMapSettings>(context, listen: false);
+    stageMapSettings.cameraAngle.value = angle;
+
+    Navigator.of(context, rootNavigator: true).pop();
+  }
+}
diff --git a/Engine/Extras/VirtualProduction/EpicStageApp/pubspec.yaml b/Engine/Extras/VirtualProduction/EpicStageApp/pubspec.yaml
--- a/Engine/Extras/VirtualProduction/EpicStageApp/pubspec.yaml
+++ b/Engine/Extras/VirtualProduction/EpicStageApp/pubspec.yaml
@@ -82,6 +82,7 @@ flutter:
 
   assets:
     - assets/images/map/
+    - assets/images/icons/
 
   fonts:
     - family: Inter
