[UnrealStage] Fix min/max modes being ignored on delta widgets. Force longitude to loop

diff --git a/Engine/Extras/VirtualProduction/EpicStageApp/lib/models/unreal_property_controller.dart b/Engine/Extras/VirtualProduction/EpicStageApp/lib/models/unreal_property_controller.dart
--- a/Engine/Extras/VirtualProduction/EpicStageApp/lib/models/unreal_property_controller.dart
+++ b/Engine/Extras/VirtualProduction/EpicStageApp/lib/models/unreal_property_controller.dart
@@ -4,16 +4,16 @@ import 'dart:collection';
 
 import 'package:flutter/foundation.dart';
 import 'package:flutter/material.dart';
+import 'package:flutter_gen/gen_l10n/app_localizations.dart';
 import 'package:logging/logging.dart';
 import 'package:provider/provider.dart';
-import 'package:flutter_gen/gen_l10n/app_localizations.dart';
 
 import '../widgets/elements/delta_widget_base.dart';
 import '../widgets/elements/unreal_widget_base.dart';
+import 'navigator_keys.dart';
 import 'property_modify_operations.dart';
 import 'unreal_property_manager.dart';
 import 'unreal_types.dart';
-import 'navigator_keys.dart';
 
 final _log = Logger('UnrealPropertyController');
 
@@ -75,12 +75,6 @@ class UnrealPropertyController<PropertyType> with ChangeNotifier {
   /// The property's maximum value as received from the engine.
   PropertyType? get engineMax => _propertyMetadata?.maxValue;
 
-  /// The overridden minimum value the controlled properties can reach. If null, [engineMin] will be used.
-  PropertyType? get overrideMin => null;
-
-  /// The overridden maximum value the controlled properties can reach. If null, [engineMax] will be used.
-  PropertyType? get overrideMax => null;
-
   /// The property's possible enum values as received from the engine.
   List<String> get propertyEnumValues => _propertyMetadata?.enumValues ?? [];
 
@@ -186,13 +180,18 @@ class UnrealPropertyController<PropertyType> with ChangeNotifier {
   /// Apply an [operation] to each property this controls by the amounts specified in [values].
   /// If [values] is null, this is treated as if it was a list of null values matching the number of properties.
   /// If [bIgnoreLimits] is true, the min and max for the widget will be ignored.
-  /// If [_bShouldInitTransaction] is false, we ought to call [beginTransaction] before calling
+  /// If [minMaxBehaviour] behaviour is provided, it will change how the value is handled when it passes the engine's
+  /// min/max values for the property.
+  /// If [overrideMin] and/or [overrideMax] are provided, they will be used instead of the min/max values provided by
+  /// the engine when applying the [minMaxBehaviour].
+  /// If [_bShouldInitTransaction] is false,[beginTransaction] must be called before calling this.
   void modifyProperties(
     PropertyModifyOperation operation, {
     List<dynamic>? values,
     bool bIgnoreLimits = false,
     PropertyMinMaxBehaviour minMaxBehaviour = PropertyMinMaxBehaviour.clamp,
-    void Function()? onChangedByUser,
+    dynamic overrideMin = null,
+    dynamic overrideMax = null,
   }) {
     assert(values == null || values.length == _properties.length);
 
@@ -203,25 +202,29 @@ class UnrealPropertyController<PropertyType> with ChangeNotifier {
         value: values?[index],
         bIgnoreLimits: bIgnoreLimits,
         minMaxBehaviour: minMaxBehaviour,
+        overrideMax: overrideMax,
+        overrideMin: overrideMin,
       );
     }
-
-    if (onChangedByUser != null) {
-      onChangedByUser();
-    }
   }
 
   /// Apply an [operation] to a single property specified by the property [index] this controls by the amounts specified
   /// in [value].
   /// If [value] is null, this is treated as null and applied to matching property [index].
   /// If [bIgnoreLimits] is true, the min and max for the widget will be ignored.
-  /// If [_bShouldInitTransaction] is false, we ought to call [beginTransaction] before calling
+  /// If [minMaxBehaviour] behaviour is provided, it will change how the value is handled when it passes the engine's
+  /// min/max values for the property.
+  /// If [overrideMin] and/or [overrideMax] are provided, they will be used instead of the min/max values provided by
+  /// the engine when applying the [minMaxBehaviour].
+  /// If [_bShouldInitTransaction] is false,[beginTransaction] must be called before calling this.
   void modifyProperty(
     PropertyModifyOperation operation,
     int index, {
     dynamic value,
     bool bIgnoreLimits = false,
     PropertyMinMaxBehaviour minMaxBehaviour = PropertyMinMaxBehaviour.clamp,
+    dynamic overrideMin = null,
+    dynamic overrideMax = null,
   }) {
     if (_bShouldInitTransaction == true) {
       if (!_bIsInTransaction) {
diff --git a/Engine/Extras/VirtualProduction/EpicStageApp/lib/widgets/elements/unreal_widget_base.dart b/Engine/Extras/VirtualProduction/EpicStageApp/lib/widgets/elements/unreal_widget_base.dart
--- a/Engine/Extras/VirtualProduction/EpicStageApp/lib/widgets/elements/unreal_widget_base.dart
+++ b/Engine/Extras/VirtualProduction/EpicStageApp/lib/widgets/elements/unreal_widget_base.dart
@@ -3,8 +3,8 @@
 import 'package:flutter/material.dart';
 
 import '../../models/property_modify_operations.dart';
-import '../../models/unreal_property_manager.dart';
 import '../../models/unreal_property_controller.dart';
+import '../../models/unreal_property_manager.dart';
 import '../../models/unreal_types.dart';
 import 'delta_widget_base.dart';
 
@@ -115,10 +115,10 @@ mixin UnrealWidgetStateMixin<WidgetType extends UnrealWidget, PropertyType> on S
   PropertyType? get engineMax => _singlePropertyController.engineMax;
 
   /// The overridden minimum value the controlled properties can reach. If null, [engineMin] will be used.
-  PropertyType? get overrideMin => _singlePropertyController.overrideMin;
+  PropertyType? get overrideMin => null;
 
   /// The overridden maximum value the controlled properties can reach. If null, [engineMax] will be used.
-  PropertyType? get overrideMax => _singlePropertyController.overrideMax;
+  PropertyType? get overrideMax => null;
 
   /// The type of modify operation this widget uses to apply deltas.
   PropertyModifyOperation get modifyOperation => _singlePropertyController.modifyOperation;
@@ -133,14 +133,14 @@ mixin UnrealWidgetStateMixin<WidgetType extends UnrealWidget, PropertyType> on S
   void initState() {
     _singlePropertyController = UnrealPropertyController(context);
     _singlePropertyController.trackAllProperties(widget.unrealProperties, widget.enableProperties);
-    _singlePropertyController.addListener(_handleOnPropertiesChanged);
+    _singlePropertyController.addListener(handleOnPropertiesChanged);
 
     super.initState();
   }
 
   @override
   void dispose() {
-    _singlePropertyController.removeListener(_handleOnPropertiesChanged);
+    _singlePropertyController.removeListener(handleOnPropertiesChanged);
     _singlePropertyController.dispose();
 
     super.dispose();
@@ -202,8 +202,8 @@ mixin UnrealWidgetStateMixin<WidgetType extends UnrealWidget, PropertyType> on S
     return SingleSharedValue(value: sharedValue, bHasMultipleValues: bHasMultipleValues);
   }
 
-  /// handle state change for [properties].
-  void _handleOnPropertiesChanged() {
+  /// Called whenever the underlying Unreal Engine property changes.
+  void handleOnPropertiesChanged() {
     setState(() {});
   }
 
@@ -213,9 +213,13 @@ mixin UnrealWidgetStateMixin<WidgetType extends UnrealWidget, PropertyType> on S
       _singlePropertyController.modifyProperties(
         modifyOperation,
         values: deltaValues,
-        onChangedByUser: widget.onChangedByUser,
+        minMaxBehaviour: widget.minMaxBehaviour,
+        overrideMin: overrideMin,
+        overrideMax: overrideMax,
       );
     });
+
+    widget.onChangedByUser?.call();
   }
 
   /// Called when the user presses the reset button.
diff --git a/Engine/Extras/VirtualProduction/EpicStageApp/lib/widgets/screens/main/tabs/details_tab.dart b/Engine/Extras/VirtualProduction/EpicStageApp/lib/widgets/screens/main/tabs/details_tab.dart
--- a/Engine/Extras/VirtualProduction/EpicStageApp/lib/widgets/screens/main/tabs/details_tab.dart
+++ b/Engine/Extras/VirtualProduction/EpicStageApp/lib/widgets/screens/main/tabs/details_tab.dart
@@ -456,7 +456,7 @@ class _DetailsTabState extends State<DetailsTab> with GuardedRefreshState {
             nonUVActorPaths,
             modifierFunction: _modifyPositionalPropertyNameBasedOnClass,
           ),
-          minMaxBehaviour: PropertyMinMaxBehaviour.ignore,
+          minMaxBehaviour: PropertyMinMaxBehaviour.clamp,
         ),
         UnrealDeltaSlider(
           key: Key('Longitude'),
@@ -465,7 +465,7 @@ class _DetailsTabState extends State<DetailsTab> with GuardedRefreshState {
             nonUVActorPaths,
             modifierFunction: _modifyPositionalPropertyNameBasedOnClass,
           ),
-          minMaxBehaviour: PropertyMinMaxBehaviour.ignore,
+          minMaxBehaviour: PropertyMinMaxBehaviour.loop,
         ),
       ]);
     } else if (uvActorPaths.isNotEmpty) {
