[LiveLinkVCAM] Force disable UI rotation while app is running

diff --git a/Engine/Extras/VirtualProduction/FlutterLiveLinkVCAM/android/app/src/main/AndroidManifest.xml b/Engine/Extras/VirtualProduction/FlutterLiveLinkVCAM/android/app/src/main/AndroidManifest.xml
--- a/Engine/Extras/VirtualProduction/FlutterLiveLinkVCAM/android/app/src/main/AndroidManifest.xml
+++ b/Engine/Extras/VirtualProduction/FlutterLiveLinkVCAM/android/app/src/main/AndroidManifest.xml
@@ -42,6 +42,7 @@
             android:exported="true"
             android:launchMode="singleTop"
             android:theme="@style/LaunchTheme"
+            android:screenOrientation="landscape"
             android:configChanges="orientation|keyboardHidden|keyboard|screenSize|smallestScreenSize|locale|layoutDirection|fontScale|screenLayout|density|uiMode"
             android:hardwareAccelerated="true"
             android:windowSoftInputMode="adjustResize">
