[LiveLinkVCAM] [UnrealStage] Add privacy policy in settings menu

diff --git a/Engine/Extras/VirtualProduction/EpicStageApp/lib/widgets/screens/main/toolbar/settings/pages/settings_dialog_main.dart b/Engine/Extras/VirtualProduction/EpicStageApp/lib/widgets/screens/main/toolbar/settings/pages/settings_dialog_main.dart
--- a/Engine/Extras/VirtualProduction/EpicStageApp/lib/widgets/screens/main/toolbar/settings/pages/settings_dialog_main.dart
+++ b/Engine/Extras/VirtualProduction/EpicStageApp/lib/widgets/screens/main/toolbar/settings/pages/settings_dialog_main.dart
@@ -23,9 +23,13 @@ class SettingsDialogMain extends StatelessWidget {
 
   const SettingsDialogMain({Key? key}) : super(key: key);
 
+  /// URI for Unreal Stage documentation
   final String _documentationUri =
       'https://dev.epicgames.com/community/learning/courses/OKO/unreal-engine-unreal-stage-ios-app/';
 
+  /// URI for privacy policy
+  final String _privacyPolicyUri = 'https://www.epicgames.com/site/en-US/privacypolicy';
+
   @override
   Widget build(BuildContext context) {
     final actorManager = Provider.of<UnrealActorManager>(context, listen: false);
@@ -72,6 +76,11 @@ class SettingsDialogMain extends StatelessWidget {
             iconPath: 'packages/epic_common/assets/icons/info.svg',
             onTap: () => rootNavigatorKey.currentState?.pushNamed(EulaScreen.route, arguments: {'onPressed': () {}}),
           ),
+          SettingsMenuItem(
+            title: localizations.settingsDialogPrivacyPolicyLabel,
+            iconPath: 'packages/epic_common/assets/icons/info.svg',
+            onTap: () => launchUrl(Uri.parse(_privacyPolicyUri)),
+          ),
           const SettingsMenuDivider(),
           Container(
             padding: EdgeInsets.symmetric(horizontal: 8),
