[LiveLinkVCAM] [UnrealStage] Consistently show scrollbars when elements are cut off

diff --git a/Engine/Extras/VirtualProduction/EpicStageApp/lib/widgets/screens/main/sidebar/outliner_panel.dart b/Engine/Extras/VirtualProduction/EpicStageApp/lib/widgets/screens/main/sidebar/outliner_panel.dart
--- a/Engine/Extras/VirtualProduction/EpicStageApp/lib/widgets/screens/main/sidebar/outliner_panel.dart
+++ b/Engine/Extras/VirtualProduction/EpicStageApp/lib/widgets/screens/main/sidebar/outliner_panel.dart
@@ -26,7 +26,6 @@ import '../../../../models/unreal_transaction_manager.dart';
 import '../../../../models/unreal_types.dart';
 import '../../../../utilities/constants.dart';
 import '../../../../utilities/guarded_refresh_state.dart';
-import '../../../elements/dropdown_button.dart';
 import '../../../elements/place_actor_menu.dart';
 import '../stage_app_main_screen.dart';
 import 'outliner_filter_menu.dart';
@@ -186,7 +185,6 @@ class _OutlinerPanelState extends State<OutlinerPanel> {
   late final OutlinerPanelSettings _outlinerSettings;
 
   final TextEditingController _searchTextController = TextEditingController();
-  final ScrollController _actorListScrollController = ScrollController();
 
   /// Set of actors that we're listening to for any updates (e.g. to class information).
   Set<UnrealObject> _listenedActors = {};
@@ -374,7 +372,7 @@ class _OutlinerPanelState extends State<OutlinerPanel> {
                 // List of actors or onboarding placeholder
                 Expanded(
                   child: _filteredActors.isNotEmpty
-                      ? ListView.builder(
+                      ? EpicListView(
                           padding: UnrealTheme.cardListViewPadding,
                           itemCount: _filteredActors.length,
                           itemBuilder: (BuildContext context, int actorIndex) {
@@ -388,7 +386,6 @@ class _OutlinerPanelState extends State<OutlinerPanel> {
                               visibilityStateCallback: (state) => setState(() => _bIsSelectedActorsVisible = state),
                             );
                           },
-                          controller: _actorListScrollController,
                         )
                       : EmptyPlaceholder(
                           message: AppLocalizations.of(context)!.outlinerEmptyMessage,
diff --git a/Engine/Extras/VirtualProduction/EpicStageApp/lib/widgets/screens/main/toolbar/settings/pages/settings_log_list.dart b/Engine/Extras/VirtualProduction/EpicStageApp/lib/widgets/screens/main/toolbar/settings/pages/settings_log_list.dart
--- a/Engine/Extras/VirtualProduction/EpicStageApp/lib/widgets/screens/main/toolbar/settings/pages/settings_log_list.dart
+++ b/Engine/Extras/VirtualProduction/EpicStageApp/lib/widgets/screens/main/toolbar/settings/pages/settings_log_list.dart
@@ -1,86 +1,22 @@
 // Copyright Epic Games, Inc. All Rights Reserved.
 
-import 'dart:io';
-
-import 'package:epic_common/logging.dart';
 import 'package:epic_common/widgets.dart';
 import 'package:flutter/material.dart';
 import 'package:flutter_gen/gen_l10n/app_localizations.dart';
 
-import '../../../../../../utilities/guarded_refresh_state.dart';
 import 'settings_log_view.dart';
 
 /// Page showing the list of log files.
-class SettingsLogList extends StatefulWidget {
+class SettingsLogList extends StatelessWidget {
   const SettingsLogList({Key? key}) : super(key: key);
 
   static const String route = '/logs';
 
-  @override
-  State<SettingsLogList> createState() => _SettingsLogListState();
-}
-
-class _SettingsLogListState extends State<SettingsLogList> with GuardedRefreshState {
-  /// Scroll controller for the scroll view of the logs.
-  final _scrollController = ScrollController();
-
-  /// List of files to display.
-  List<File> _logFiles = [];
-
-  @override
-  void initState() {
-    super.initState();
-
-    _updateLogFileList();
-  }
-
-  @override
-  void dispose() {
-    _scrollController.dispose();
-    super.dispose();
-  }
-
-  /// Update the list of log files.
-  void _updateLogFileList() async {
-    _logFiles = await Logging.getAppLogFiles();
-    setState(() {});
-  }
-
   @override
   Widget build(BuildContext context) {
-    return SettingsPageScaffold(
-      title: AppLocalizations.of(context)!.settingsDialogApplicationLogLabel,
-      body: SizedBox(
-        height: 240,
-        child: MediaQuery.removePadding(
-          removeTop: true,
-          removeBottom: true,
-          context: context,
-          child: EpicListView(
-            itemCount: _logFiles.length,
-            itemBuilder: (context, index) {
-              final File file = _logFiles[index];
-              final bool bIsCurrent = index == 0;
-
-              return SettingsMenuItem(
-                title: Logging.makeNameForLog(
-                  context: context,
-                  logFile: file,
-                  bIsCurrent: bIsCurrent,
-                ),
-                iconPath: 'packages/epic_common/assets/icons/log.svg',
-                onTap: () => Navigator.of(context).pushNamed(
-                  SettingsLogView.route,
-                  arguments: SettingsLogViewArguments(
-                    file: file,
-                    bIsCurrent: bIsCurrent,
-                  ),
-                ),
-              );
-            },
-          ),
-        ),
-      ),
+    return LogListSettingsPage(
+      AppLocalizations.of(context)!.settingsDialogApplicationLogLabel,
+      SettingsLogView.route,
     );
   }
 }
diff --git a/Engine/Extras/VirtualProduction/EpicStageApp/lib/widgets/screens/main/toolbar/settings/pages/settings_log_view.dart b/Engine/Extras/VirtualProduction/EpicStageApp/lib/widgets/screens/main/toolbar/settings/pages/settings_log_view.dart
--- a/Engine/Extras/VirtualProduction/EpicStageApp/lib/widgets/screens/main/toolbar/settings/pages/settings_log_view.dart
+++ b/Engine/Extras/VirtualProduction/EpicStageApp/lib/widgets/screens/main/toolbar/settings/pages/settings_log_view.dart
@@ -1,75 +1,11 @@
 // Copyright Epic Games, Inc. All Rights Reserved.
 
-import 'dart:io';
-
-import 'package:epic_common/logging.dart';
 import 'package:epic_common/widgets.dart';
 import 'package:flutter/material.dart';
-import 'package:share_plus/share_plus.dart';
-
-import '../../../../../../utilities/guarded_refresh_state.dart';
-
-/// Arguments passed to the route when a [SettingsLogView] is pushed to the navigation stack.
-class SettingsLogViewArguments {
-  const SettingsLogViewArguments({required this.file, required this.bIsCurrent});
-
-  /// The file to open.
-  final File file;
-
-  /// Whether this is the current log.
-  final bool bIsCurrent;
-}
 
 /// Page showing the full view of a log.
-class SettingsLogView extends StatefulWidget {
+class SettingsLogView extends LogViewSettingsPage {
   const SettingsLogView({Key? key}) : super(key: key);
 
   static const String route = '/log_view';
-
-  @override
-  State<SettingsLogView> createState() => _SettingsLogViewState();
-}
-
-class _SettingsLogViewState extends State<SettingsLogView> with GuardedRefreshState {
-  final shareButtonKey = GlobalKey();
-
-  @override
-  Widget build(BuildContext context) {
-    final arguments = ModalRoute.of(context)?.settings.arguments as SettingsLogViewArguments;
-
-    return SettingsPageScaffold(
-      title: Logging.makeNameForLog(
-        context: context,
-        logFile: arguments.file,
-        bIsCurrent: arguments.bIsCurrent,
-      ),
-      titleBarTrailing: EpicIconButton(
-        key: shareButtonKey,
-        iconPath: 'packages/epic_common/assets/icons/share.svg',
-        color: Theme.of(context).colorScheme.primary,
-        onPressed: () => _shareLog(arguments.file),
-      ),
-      body: SizedBox(
-        height: MediaQuery.of(context).size.height - 120,
-        child: LogFileViewer(arguments.file),
-      ),
-    );
-  }
-
-  /// Open a prompt to share the log's contents using the local operating system.
-  void _shareLog(File log) async {
-    // Need to provide this position for sharing to work on iPad
-    final renderBox = shareButtonKey.currentContext?.findRenderObject() as RenderBox?;
-    if (renderBox == null) {
-      return;
-    }
-
-    final sharePosition = renderBox.localToGlobal(Offset.zero) & renderBox.size;
-
-    Share.share(
-      await log.readAsString(),
-      subject: log.uri.pathSegments.last,
-      sharePositionOrigin: sharePosition,
-    );
-  }
 }
