[UnrealStage] Add camera focus/DoF tab

diff --git a/Engine/Extras/VirtualProduction/EpicStageApp/lib/main.dart b/Engine/Extras/VirtualProduction/EpicStageApp/lib/main.dart
--- a/Engine/Extras/VirtualProduction/EpicStageApp/lib/main.dart
+++ b/Engine/Extras/VirtualProduction/EpicStageApp/lib/main.dart
@@ -141,6 +141,10 @@ class _EpicStageAppState extends State<EpicStageApp> with TickerProviderStateMix
       case AppLifecycleState.resumed:
         _log.info('App resumed');
         break;
+
+      case AppLifecycleState.hidden:
+        _log.info('App hidden');
+        break;
     }
   }
 
diff --git a/Engine/Extras/VirtualProduction/EpicStageApp/lib/models/preview_render_manager.dart b/Engine/Extras/VirtualProduction/EpicStageApp/lib/models/preview_render_manager.dart
--- a/Engine/Extras/VirtualProduction/EpicStageApp/lib/models/preview_render_manager.dart
+++ b/Engine/Extras/VirtualProduction/EpicStageApp/lib/models/preview_render_manager.dart
@@ -660,7 +660,7 @@ class PreviewActorData {
 }
 
 /// Mixin for classes that want to receive renders from [PreviewRenderManager].
-abstract class PreviewRenderConsumer {
+mixin PreviewRenderConsumer {
   /// Called when a preview render from the engine completes.
   /// [bIsForLatestRequest] is true if this contains the latest available preview and actor position data, or false if
   /// the engine is expected to send more recent data after this.
diff --git a/Engine/Extras/VirtualProduction/EpicStageApp/lib/models/settings/camera_settings_tab_settings.dart b/Engine/Extras/VirtualProduction/EpicStageApp/lib/models/settings/camera_settings_tab_settings.dart
new file mode 100644
--- /dev/null
+++ b/Engine/Extras/VirtualProduction/EpicStageApp/lib/models/settings/camera_settings_tab_settings.dart
@@ -0,0 +1,30 @@
+// Copyright Epic Games, Inc. All Rights Reserved.
+
+import 'package:epic_common/preferences.dart';
+import 'package:streaming_shared_preferences/streaming_shared_preferences.dart';
+
+/// Possible modes for the camera settings tab's left column.
+enum CameraSettingsLeftColumnMode {
+  camera,
+  focus,
+}
+
+/// Holds user settings used in the Camera Settings tab.
+class CameraSettingsTabSettings {
+  CameraSettingsTabSettings(PreferencesBundle preferences)
+      : selectedCamera = preferences.transient.get(
+          'cameraSettingsTab.selectedCamera',
+          defaultValue: '',
+        ),
+        leftColumnMode = preferences.persistent.getEnum(
+          'cameraSettingsTab.mode',
+          defaultValue: CameraSettingsLeftColumnMode.focus,
+          enumValues: CameraSettingsLeftColumnMode.values,
+        );
+
+  /// The path of the camera currently selected for editing in the camera settings tab.
+  final TransientPreference<String> selectedCamera;
+
+  /// The selected mode for the camera settings tab's left column.
+  final Preference<CameraSettingsLeftColumnMode> leftColumnMode;
+}
diff --git a/Engine/Extras/VirtualProduction/EpicStageApp/lib/utilities/net_utilities.dart b/Engine/Extras/VirtualProduction/EpicStageApp/lib/utilities/net_utilities.dart
--- a/Engine/Extras/VirtualProduction/EpicStageApp/lib/utilities/net_utilities.dart
+++ b/Engine/Extras/VirtualProduction/EpicStageApp/lib/utilities/net_utilities.dart
@@ -8,6 +8,7 @@ import 'package:epic_common/unreal_beacon.dart';
 import 'package:logging/logging.dart';
 
 import '../models/engine_connection.dart';
+import '../models/unreal_types.dart';
 
 /// The configuration to use for the beacon that detects compatible Unreal Engine instances.
 final unrealEngineBeaconConfig = UnrealEngineBeaconConfig(
@@ -230,3 +231,47 @@ class MessageTimeoutRetryHelper {
     }
   }
 }
+
+/// Given a list of [componentPaths], query the engine via the [connectionManager] and create local UnrealObjects for
+/// all components whose class path exactly matches the [classPath].
+Future<List<UnrealObject>> getComponentsOfType({
+  required EngineConnectionManager connectionManager,
+  required List<String> componentPaths,
+  required String classPath,
+}) async {
+  // Request description of each object so we can determine their classes and names
+  final List<UnrealHttpRequest> describeRequests = [];
+  for (final dynamic componentPath in componentPaths) {
+    if (componentPath is! String) {
+      continue;
+    }
+
+    describeRequests.add(UnrealHttpRequest(url: '/remote/object/describe', verb: 'PUT', body: {
+      'objectPath': componentPath,
+      'access': 'READ_ACCESS',
+    }));
+  }
+
+  final UnrealHttpResponse batchDescribeResponse = await connectionManager.sendBatchedHttpRequest(describeRequests);
+  if (batchDescribeResponse.code != 200) {
+    return [];
+  }
+
+  // Create UnrealObjects only for the objects whose classes match
+  final List<UnrealObject> cameras = [];
+  for (int componentIndex = 0; componentIndex < componentPaths.length; ++componentIndex) {
+    final UnrealHttpResponse? describeResponse = batchDescribeResponse.body[componentIndex];
+    if (describeResponse == null || describeResponse.code != 200) {
+      continue;
+    }
+
+    if (describeResponse.body?['Class'] == classPath) {
+      cameras.add(UnrealObject(
+        path: componentPaths[componentIndex],
+        name: describeResponse.body?['Name'],
+      ));
+    }
+  }
+
+  return cameras;
+}
diff --git a/Engine/Extras/VirtualProduction/EpicStageApp/lib/utilities/unreal_utilities.dart b/Engine/Extras/VirtualProduction/EpicStageApp/lib/utilities/unreal_utilities.dart
--- a/Engine/Extras/VirtualProduction/EpicStageApp/lib/utilities/unreal_utilities.dart
+++ b/Engine/Extras/VirtualProduction/EpicStageApp/lib/utilities/unreal_utilities.dart
@@ -121,6 +121,8 @@ Map<String, _UnrealTypeConversionData> _typeConverters = {
     },
     dartToJson: (dynamic dartValue, Map<String, dynamic> conversionMetadata) => dartValue.toJson(),
   ),
+  'TSoftObjectPtr<ACineCameraActor>': _noConversion,
+  'TSoftObjectPtr<AActor>': _noConversion,
   unrealEnumTypeName: _noConversion,
 };
 
diff --git a/Engine/Extras/VirtualProduction/EpicStageApp/lib/widgets/elements/actor_reference_viewer.dart b/Engine/Extras/VirtualProduction/EpicStageApp/lib/widgets/elements/actor_reference_viewer.dart
new file mode 100644
--- /dev/null
+++ b/Engine/Extras/VirtualProduction/EpicStageApp/lib/widgets/elements/actor_reference_viewer.dart
@@ -0,0 +1,156 @@
+// Copyright Epic Games, Inc. All Rights Reserved.
+
+import 'package:epic_common/preferences.dart';
+import 'package:epic_common/theme.dart';
+import 'package:flutter/material.dart';
+import 'package:flutter_gen/gen_l10n/app_localizations.dart';
+import 'package:provider/provider.dart';
+
+import '../../models/engine_connection.dart';
+import '../../models/settings/delta_widget_settings.dart';
+import 'delta_widget_base.dart';
+import 'unreal_widget_base.dart';
+
+/// Displays the name of actors referenced by properties in Unreal Engine.
+/// This widget is read-only, so users have to change the value from the editor.
+class UnrealActorReferenceViewer extends UnrealWidget {
+  const UnrealActorReferenceViewer({
+    super.key,
+    required super.unrealProperties,
+  });
+
+  @override
+  State<UnrealActorReferenceViewer> createState() => _UnrealActorDisplayState();
+}
+
+class _UnrealActorDisplayState extends State<UnrealActorReferenceViewer>
+    with UnrealWidgetStateMixin<UnrealActorReferenceViewer, String> {
+  /// The name to display for the selected actor(s).
+  String _actorName = '';
+
+  @override
+  Widget build(BuildContext context) {
+    final TextStyle textStyle = Theme.of(context).textTheme.labelMedium!;
+
+    return Padding(
+      padding: EdgeInsets.only(bottom: 16),
+      child: TransientPreferenceBuilder(
+        preference: Provider.of<DeltaWidgetSettings>(context, listen: false).bIsInResetMode,
+        builder: (context, final bool bIsInResetMode) => Row(
+          children: [
+            Expanded(
+              child: Padding(
+                padding: EdgeInsets.symmetric(horizontal: DeltaWidgetConstants.widgetOuterXPadding),
+                child: Column(
+                  mainAxisAlignment: MainAxisAlignment.start,
+                  crossAxisAlignment: CrossAxisAlignment.start,
+                  children: [
+                    Text(
+                      propertyLabel,
+                      style: textStyle.copyWith(
+                        color: textStyle.color!.withOpacity(bIsInResetMode ? 0.4 : 1.0),
+                      ),
+                    ),
+                    Padding(
+                      padding: const EdgeInsets.only(
+                        left: DeltaWidgetConstants.widgetInnerXPadding,
+                        right: DeltaWidgetConstants.widgetInnerXPadding,
+                        top: 4,
+                      ),
+                      child: Container(
+                        height: 36,
+                        decoration: BoxDecoration(
+                          color: bIsInResetMode ? UnrealColors.gray22 : Theme.of(context).colorScheme.background,
+                          borderRadius: BorderRadius.all(
+                            const Radius.circular(4),
+                          ),
+                        ),
+                        padding: EdgeInsets.symmetric(horizontal: 12),
+                        child: Row(
+                          children: [
+                            Expanded(
+                              child: Text(
+                                _actorName,
+                                overflow: TextOverflow.ellipsis,
+                                maxLines: 1,
+                                style: Theme.of(context).textTheme.bodyMedium!.copyWith(
+                                      color: bIsInResetMode ? UnrealColors.gray56 : UnrealColors.white,
+                                      letterSpacing: 0.5,
+                                    ),
+                              ),
+                            ),
+                          ],
+                        ),
+                      ),
+                    ),
+                  ],
+                ),
+              ),
+            ),
+          ],
+        ),
+      ),
+    );
+  }
+
+  @override
+  void handleOnPropertiesChanged() {
+    setState(() {
+      _actorName = '';
+    });
+
+    _updateDisplayName();
+  }
+
+  /// Retrieve the name of the referenced actor from the engine.
+  void _updateDisplayName() async {
+    final SingleSharedValue<String> actorPath = getSingleSharedValue();
+
+    if (actorPath.bHasMultipleValues) {
+      setState(() {
+        _actorName = AppLocalizations.of(context)!.mismatchedValuesLabel;
+      });
+      return;
+    }
+
+    if (actorPath.value == null || actorPath.value!.isEmpty) {
+      setState(() {
+        _actorName = AppLocalizations.of(context)!.actorReferenceNoneLabel;
+      });
+
+      return;
+    }
+
+    final connection = Provider.of<EngineConnectionManager>(context, listen: false);
+    final UnrealHttpResponse response = await connection.sendHttpRequest(UnrealHttpRequest(
+      url: '/remote/object/call',
+      verb: 'PUT',
+      body: {
+        'objectPath': '/Script/Engine.Default__KismetSystemLibrary',
+        'functionName': 'GetDisplayName',
+        'parameters': {
+          'Object': actorPath.value!,
+        },
+      },
+    ));
+
+    if (response.code != HttpResponseCode.ok) {
+      setState(() {
+        _actorName = AppLocalizations.of(context)!.actorReferenceInvalidLabel;
+      });
+      return;
+    }
+
+    final cineCameraName = response.body['ReturnValue'];
+    if (cineCameraName is! String) {
+      setState(() {
+        _actorName = AppLocalizations.of(context)!.actorReferenceInvalidLabel;
+      });
+      return;
+    }
+
+    setState(() {
+      _actorName = cineCameraName;
+    });
+  }
+}
diff --git a/Engine/Extras/VirtualProduction/EpicStageApp/lib/widgets/elements/delta_slider.dart b/Engine/Extras/VirtualProduction/EpicStageApp/lib/widgets/elements/delta_slider.dart
--- a/Engine/Extras/VirtualProduction/EpicStageApp/lib/widgets/elements/delta_slider.dart
+++ b/Engine/Extras/VirtualProduction/EpicStageApp/lib/widgets/elements/delta_slider.dart
@@ -35,8 +35,10 @@ class UnrealDeltaSlider extends UnrealWidget {
     this.softMax,
     this.hardMin,
     this.hardMax,
+    this.exponent = 1.0,
     this.bShowResetButton = true,
     this.fillColor,
+    this.valueMessageOverride,
   });
 
   /// If provided, call this function to build the widget displayed as the slider's label. Otherwise, the label will be
@@ -61,12 +63,19 @@ class UnrealDeltaSlider extends UnrealWidget {
   /// This changes both the max value displayed on the slider and prevents sliding past this value.
   final double? hardMax;
 
+  /// Exponent value applied to the slider's linear position. Higher exponents make the slider's effective delta value
+  /// smaller at smaller slider values.
+  final double exponent;
+
   /// If true, show the reset button.
   final bool bShowResetButton;
 
   /// The color to show in the filled portion of the bar (left of the value indicators).
   final Color? fillColor;
 
+  /// If set, replace the value display with this string.
+  final String? valueMessageOverride;
+
   @override
   _UnrealDeltaSliderState createState() => _UnrealDeltaSliderState();
 }
@@ -90,11 +99,13 @@ class _UnrealDeltaSliderState extends State<UnrealDeltaSlider> with UnrealWidget
         onValuesSet: _onValuesSet,
         onReset: handleOnResetByUser,
         onInteractionFinished: endTransaction,
+        exponent: widget.exponent,
         min: widget.softMin ?? widget.hardMin ?? engineMin,
         max: widget.softMax ?? widget.hardMax ?? engineMax,
         label: propertyLabel,
         buildLabel: widget.buildLabel,
         fillColor: widget.fillColor,
+        valueMessageOverride: widget.valueMessageOverride,
       ),
     );
   }
@@ -124,6 +135,7 @@ class DrivenDeltaSlider extends StatefulWidget {
     this.onReset,
     this.onInteractionFinished,
     this.fillColor,
+    this.valueMessageOverride,
     this.baseSensitivity,
   }) : super(key: key);
 
@@ -177,6 +189,9 @@ class DrivenDeltaSlider extends StatefulWidget {
   /// The color to show in the filled portion of the bar (left of the value indicators).
   final Color? fillColor;
 
+  /// If set, replace the value display with this string.
+  final String? valueMessageOverride;
+
   /// How sensitive the slider should be when the user's sensitivity is set to 1.0.
   final double? baseSensitivity;
 
@@ -267,7 +282,7 @@ class _DrivenDeltaSliderState extends State<DrivenDeltaSlider> with DeltaWidgetS
 
   @override
   Widget build(BuildContext context) {
-    final String? valueString = _singleSharedValueString;
+    final String? valueString = widget.valueMessageOverride ?? _singleSharedValueString;
 
     // Build the list of slider values to pass to the slider widget
     final List<_DeltaSliderUIValue> uiValues = [];
@@ -417,7 +432,7 @@ class _DrivenDeltaSliderState extends State<DrivenDeltaSlider> with DeltaWidgetS
         final double exponentialValue = _exponentiateNormalizedValue(normalizedValue);
         final double exponentialResultValue = (exponentialValue + normalizedDelta).clamp(0, 1);
         final double linearResultValue = exponentiateValue(exponentialResultValue, widget.exponent, false);
-        deltaValues[valueIndex] = linearResultValue - normalizedValue;
+        deltaValues[valueIndex] = (linearResultValue - normalizedValue) * _valueSpan;
       }
     }
 
diff --git a/Engine/Extras/VirtualProduction/EpicStageApp/lib/widgets/elements/dropdown_button.dart b/Engine/Extras/VirtualProduction/EpicStageApp/lib/widgets/elements/dropdown_button.dart
--- a/Engine/Extras/VirtualProduction/EpicStageApp/lib/widgets/elements/dropdown_button.dart
+++ b/Engine/Extras/VirtualProduction/EpicStageApp/lib/widgets/elements/dropdown_button.dart
@@ -3,6 +3,7 @@
 import 'package:epic_common/preferences.dart';
 import 'package:epic_common/widgets.dart';
 import 'package:flutter/material.dart';
+import 'package:flutter_gen/gen_l10n/app_localizations.dart';
 import 'package:provider/provider.dart';
 
 import '../../models/property_modify_operations.dart';
@@ -23,18 +24,79 @@ class UnrealDropdownSelector extends UnrealWidget {
   _UnrealDropdownButtonState createState() => _UnrealDropdownButtonState();
 }
 
-class _UnrealDropdownButtonState extends State<UnrealDropdownSelector>
-    with UnrealWidgetStateMixin<UnrealDropdownSelector, String> {
+class _UnrealDropdownButtonState extends _UnrealDropdownSelectorOuterState<UnrealDropdownSelector, String> {
+  @override
+  void validateValue(dynamic value) {
+    assert(value is String);
+    assert(propertyEnumValues.contains(value));
+  }
+
+  @override
+  Widget buildDropdownSelector({
+    required bool bIsInResetMode,
+    required SingleSharedValue<String> sharedValue,
+    required String hint,
+  }) =>
+      DropdownSelector<String>(
+        items: propertyEnumValues,
+        value: sharedValue.value,
+        onChanged: onChanged,
+        hint: hint,
+        bIsEnabled: !bIsInResetMode,
+      );
+}
+
+/// A dropdown that controls one or more boolean properties in Unreal Engine.
+class UnrealBooleanDropdownSelector extends UnrealWidget {
+  const UnrealBooleanDropdownSelector({
+    super.key,
+    super.overrideName,
+    super.enableProperties,
+    required super.unrealProperties,
+  });
+
+  @override
+  _UnrealBooleanDropdownSelectorState createState() => _UnrealBooleanDropdownSelectorState();
+}
+
+class _UnrealBooleanDropdownSelectorState
+    extends _UnrealDropdownSelectorOuterState<UnrealBooleanDropdownSelector, bool> {
+  @override
+  void validateValue(dynamic value) {
+    assert(value is bool);
+  }
+
+  @override
+  Widget buildDropdownSelector({
+    required bool bIsInResetMode,
+    required SingleSharedValue<bool> sharedValue,
+    required String hint,
+  }) =>
+      DropdownSelector<bool>(
+        items: [false, true],
+        makeItemName: (item) => switch (item) {
+          true => AppLocalizations.of(context)!.dropdownSelectorEnabledLabel,
+          false => AppLocalizations.of(context)!.dropdownSelectorDisabledLabel,
+        },
+        value: sharedValue.value,
+        onChanged: onChanged,
+        hint: hint,
+        bIsEnabled: !bIsInResetMode,
+      );
+}
+
+abstract class _UnrealDropdownSelectorOuterState<WidgetType extends UnrealWidget, PropertyType>
+    extends State<WidgetType> with UnrealWidgetStateMixin<WidgetType, PropertyType> {
   @override
   PropertyModifyOperation get modifyOperation => const SetOperation();
 
   @override
   Widget build(BuildContext context) {
-    final SingleSharedValue<String> sharedValue = getSingleSharedValue();
+    final SingleSharedValue<PropertyType> sharedValue = getSingleSharedValue();
     final TextStyle textStyle = Theme.of(context).textTheme.labelMedium!;
 
     return Padding(
-      padding: EdgeInsets.only(bottom: 8),
+      padding: EdgeInsets.only(bottom: 16),
       child: TransientPreferenceBuilder(
         preference: Provider.of<DeltaWidgetSettings>(context, listen: false).bIsInResetMode,
         builder: (context, final bool bIsInResetMode) => Row(
@@ -58,12 +120,10 @@ class _UnrealDropdownButtonState extends State<UnrealDropdownSelector>
                         right: DeltaWidgetConstants.widgetInnerXPadding,
                         top: 4,
                       ),
-                      child: DropdownSelector<String>(
-                        items: propertyEnumValues,
-                        value: sharedValue.value,
-                        onChanged: _onChanged,
-                        hint: sharedValue.bHasMultipleValues ? 'Multiple values' : '',
-                        bIsEnabled: !bIsInResetMode,
+                      child: buildDropdownSelector(
+                        bIsInResetMode: bIsInResetMode,
+                        sharedValue: sharedValue,
+                        hint: sharedValue.bHasMultipleValues ? AppLocalizations.of(context)!.mismatchedValuesLabel : '',
                       ),
                     ),
                   ],
@@ -81,9 +141,22 @@ class _UnrealDropdownButtonState extends State<UnrealDropdownSelector>
     );
   }
 
-  void _onChanged(dynamic newValue) {
-    assert(newValue is String);
-    assert(propertyEnumValues.contains(newValue));
+  /// Validates the user-provided value before it's set.
+  @protected
+  void validateValue(dynamic value) {}
+
+  /// Build the inner dropdown selector, taking into account whether the widget [bIsInResetMode], the current
+  /// [sharedValue] between all multi-selected objects, and the [hint] to display if the shared value is null.
+  Widget buildDropdownSelector({
+    required bool bIsInResetMode,
+    required SingleSharedValue<PropertyType> sharedValue,
+    required String hint,
+  });
+
+  /// Called when the user changes the value using the dropdown.
+  @protected
+  void onChanged(dynamic newValue) {
+    validateValue(newValue);
 
     handleOnChangedByUser(List.filled(widget.unrealProperties.length, newValue));
     endTransaction();
diff --git a/Engine/Extras/VirtualProduction/EpicStageApp/lib/widgets/screens/main/tabs/camera_settings_tab.dart b/Engine/Extras/VirtualProduction/EpicStageApp/lib/widgets/screens/main/tabs/camera_settings_tab.dart
new file mode 100644
--- /dev/null
+++ b/Engine/Extras/VirtualProduction/EpicStageApp/lib/widgets/screens/main/tabs/camera_settings_tab.dart
@@ -0,0 +1,962 @@
+// Copyright Epic Games, Inc. All Rights Reserved.
+
+import 'dart:async';
+
+import 'package:epic_common/preferences.dart';
+import 'package:epic_common/theme.dart';
+import 'package:epic_common/widgets.dart';
+import 'package:flutter/material.dart';
+import 'package:flutter_gen/gen_l10n/app_localizations.dart';
+import 'package:logging/logging.dart';
+import 'package:provider/provider.dart';
+import 'package:streaming_shared_preferences/streaming_shared_preferences.dart';
+
+import '../../../../models/engine_connection.dart';
+import '../../../../models/property_modify_operations.dart';
+import '../../../../models/settings/camera_settings_tab_settings.dart';
+import '../../../../models/settings/main_screen_settings.dart';
+import '../../../../models/settings/selected_actor_settings.dart';
+import '../../../../models/unreal_actor_manager.dart';
+import '../../../../models/unreal_property_controller.dart';
+import '../../../../models/unreal_types.dart';
+import '../../../../utilities/constants.dart';
+import '../../../../utilities/guarded_refresh_state.dart';
+import '../../../../utilities/net_utilities.dart';
+import '../../../elements/actor_reference_viewer.dart';
+import '../../../elements/delta_slider.dart';
+import '../../../elements/dropdown_button.dart';
+import '../../../elements/reset_mode_button.dart';
+import '../sidebar/outliner_panel.dart';
+
+final _log = Logger('CameraTab');
+
+/// Tab to control camera settings.
+class CameraSettingsTab extends StatefulWidget {
+  const CameraSettingsTab({Key? key}) : super(key: key);
+
+  static const String iconPath = 'assets/images/icons/cine_camera.svg';
+
+  static String getTitle(BuildContext context) => AppLocalizations.of(context)!.tabTitleCamera;
+
+  @override
+  State<StatefulWidget> createState() => _CameraSettingsTabState();
+}
+
+class _CameraSettingsTabState extends State<CameraSettingsTab> with GuardedRefreshState {
+  late final EngineConnectionManager _connectionManager;
+  late final CameraSettingsTabSettings _tabSettings;
+
+  /// How often to refresh the list of available cameras.
+  static const Duration _refreshRate = Duration(seconds: 3);
+
+  /// The list of entries to display in the outliner. Stored here so we can access their cached information.
+  final List<_CameraSettingsOutlinerEntryData> _entries = [];
+
+  /// Timer used to refresh the list of available cameras.
+  Timer? _refreshTimer;
+
+  /// If true, a refresh is in progress.
+  Future? _bPendingRefresh = null;
+
+  @override
+  void initState() {
+    super.initState();
+
+    _connectionManager = Provider.of<EngineConnectionManager>(context, listen: false);
+    _tabSettings = CameraSettingsTabSettings(PreferencesBundle.of(context));
+
+    // Whenever the selected camera changes, refresh so children redraw with the correct properties
+    _tabSettings.selectedCamera.listen(refreshOnData);
+
+    _refreshTimer = Timer.periodic(_refreshRate, (timer) => _startRefreshingTargetList());
+    _startRefreshingTargetList();
+  }
+
+  @override
+  void dispose() {
+    _refreshTimer?.cancel();
+
+    super.dispose();
+  }
+
+  @override
+  Widget build(BuildContext context) {
+    return Provider(
+      create: (_) => _tabSettings,
+      child: Padding(
+        padding: EdgeInsets.all(UnrealTheme.cardMargin),
+        child: Row(
+          children: [
+            // Main controls
+            Expanded(
+              child: Card(
+                child: Column(
+                  children: [
+                    CardLargeHeader(
+                      iconPath: CameraSettingsTab.iconPath,
+                      title: _getTitle(),
+                      subtitle: AppLocalizations.of(context)!.tabTitleCamera,
+                      trailing: const ResetModeButton(),
+                    ),
+                    Expanded(
+                      child: _CameraSettingsMainControls(
+                        selectedCamera: _getEntryDataForSelectedObject(),
+                      ),
+                    ),
+                  ],
+                ),
+              ),
+            ),
+
+            PreferenceBuilder(
+              preference: Provider.of<MainScreenSettings>(context).bIsOutlinerPanelOpen,
+              builder: (context, final bool bIsOutlinerPanelOpen) {
+                if (!bIsOutlinerPanelOpen) {
+                  return const SizedBox();
+                }
+
+                return Row(children: [
+                  SizedBox(width: UnrealTheme.cardMargin),
+
+                  // Outliner and target panels
+                  SizedBox(
+                    width: getOutlinerWidth(context),
+                    child: Card(
+                      child: _CameraSettingsOutlinerPanel(
+                        entries: _entries,
+                      ),
+                    ),
+                  ),
+                ]);
+              },
+            ),
+          ],
+        ),
+      ),
+    );
+  }
+
+  /// Get the title to display at the top of the tab.
+  String? _getTitle() => _getEntryDataForSelectedObject()?.name;
+
+  /// Start refreshing the target list if there isn't one in progress.
+  void _startRefreshingTargetList() {
+    if (!mounted) {
+      return;
+    }
+
+    if (_bPendingRefresh == null) {
+      _bPendingRefresh = _refreshTargetList().then((_) => _bPendingRefresh = null);
+    }
+  }
+
+  /// Refresh the list of possible camera settings targets.
+  Future<void> _refreshTargetList() async {
+    _entries.clear();
+
+    final selectedActorSettings = Provider.of<SelectedActorSettings>(context, listen: false);
+    final actorManager = Provider.of<UnrealActorManager>(context, listen: false);
+
+    // Wait until we have the nDisplay root actors available so we can retrieve the selected one
+    await actorManager.getInitialActorsOfClass(nDisplayRootActorClassName);
+
+    final String displayClusterRootPath = selectedActorSettings.displayClusterRootPath.getValue();
+    final UnrealObject? rootActor = actorManager.getActorAtPath(displayClusterRootPath);
+
+    if (displayClusterRootPath.isNotEmpty && rootActor != null) {
+      await _addTargetsForConfigActorPath(rootActor);
+    }
+
+    if (!mounted) {
+      return;
+    }
+
+    // If the current target isn't valid anymore, deselect it
+    if (_getEntryDataForSelectedObject() == null) {
+      final String newSelectedCamera = _entries
+              .where((entryData) => entryData.type == _CameraSettingsOutlinerEntryType.icvfxCamera)
+              .firstOrNull
+              ?.path ??
+          '';
+
+      _tabSettings.selectedCamera.setValue(newSelectedCamera);
+    }
+
+    setState(() {});
+  }
+
+  /// Request data about an nDisplay config and add its camera targets to the target list.
+  Future<_CameraSettingsOutlinerEntryData?> _addTargetsForConfigActorPath(UnrealObject configActor) async {
+    // Add an entry for the root actor
+    final rootEntry = _CameraSettingsOutlinerEntryData(
+      path: configActor.path,
+      name: configActor.name,
+      type: _CameraSettingsOutlinerEntryType.nDisplayConfig,
+    );
+    _entries.add(rootEntry);
+
+    // Request the root actor's list of components, which may include cameras
+    final UnrealHttpResponse response = await _connectionManager.sendHttpRequest(UnrealHttpRequest(
+      url: '/remote/object/property',
+      verb: 'PUT',
+      body: {
+        'objectPath': configActor.path,
+        'propertyName': 'BlueprintCreatedComponents',
+        'access': 'READ_ACCESS',
+      },
+    ));
+
+    await _addEntriesForComponentsResponse(response, rootEntry);
+
+    return rootEntry;
+  }
+
+  /// Given a response from querying an nDisplay root actor's configuration data, add all of its color grading targets
+  /// to the target list.
+  Future<void> _addEntriesForComponentsResponse(
+    UnrealHttpResponse? response,
+    _CameraSettingsOutlinerEntryData parent,
+  ) async {
+    if (response?.code != HttpResponseCode.ok || !mounted) {
+      return;
+    }
+
+    final List<dynamic>? componentPaths = response!.body?['BlueprintCreatedComponents'];
+    if (componentPaths == null) {
+      return;
+    }
+
+    final List<UnrealObject> cameras = await getComponentsOfType(
+      connectionManager: _connectionManager,
+      componentPaths: componentPaths.map((e) => e.toString()).toList(growable: false),
+      classPath: '/Script/DisplayCluster.DisplayClusterICVFXCameraComponent',
+    );
+
+    for (final UnrealObject camera in cameras) {
+      _entries.add(_CameraSettingsOutlinerEntryData(
+        path: camera.path,
+        name: camera.name,
+        type: _CameraSettingsOutlinerEntryType.icvfxCamera,
+        parent: parent,
+      ));
+    }
+  }
+
+  /// Get the data for the given object path.
+  _CameraSettingsOutlinerEntryData? _getEntryDataForObject(String objectPath) {
+    for (final _CameraSettingsOutlinerEntryData entryData in _entries) {
+      if (entryData.path == objectPath) {
+        return entryData;
+      }
+    }
+
+    return null;
+  }
+
+  /// Get the data for the currently selected entry.
+  _CameraSettingsOutlinerEntryData? _getEntryDataForSelectedObject() {
+    final String targetPath = _tabSettings.selectedCamera.getValue();
+    if (targetPath.isEmpty) {
+      return null;
+    }
+
+    return _getEntryDataForObject(targetPath);
+  }
+}
+
+/// The main controls of the camera settings tab.
+class _CameraSettingsMainControls extends StatelessWidget {
+  const _CameraSettingsMainControls({required this.selectedCamera});
+
+  /// The camera selected for editing.
+  final _CameraSettingsOutlinerEntryData? selectedCamera;
+
+  /// Height of the header area in each column.
+  final double _headerHeight = 52;
+
+  @override
+  Widget build(BuildContext context) {
+    if (selectedCamera == null) {
+      return const SizedBox();
+    }
+
+    final settings = Provider.of<CameraSettingsTabSettings>(context);
+    final localizations = AppLocalizations.of(context)!;
+
+    return Padding(
+      padding: const EdgeInsets.only(top: UnrealTheme.sectionMargin),
+      child: Row(
+        crossAxisAlignment: CrossAxisAlignment.end,
+        children: [
+          // Camera/focus settings
+          Expanded(
+            child: Container(
+              color: Theme.of(context).colorScheme.surface,
+              child: Column(
+                children: [
+                  Container(
+                    color: Theme.of(context).colorScheme.surfaceTint,
+                    height: _headerHeight,
+                    child: Center(
+                      child: PreferenceBuilder(
+                        preference: settings.leftColumnMode,
+                        builder: (context, CameraSettingsLeftColumnMode mode) =>
+                            SelectorBar<CameraSettingsLeftColumnMode>(
+                          value: mode,
+                          valueNames: {
+                            CameraSettingsLeftColumnMode.camera: localizations.cameraTabCameraPropertiesLabel,
+                            CameraSettingsLeftColumnMode.focus: localizations.cameraTabFocusPropertiesLabel,
+                          },
+                          onSelected: (newMode) => settings.leftColumnMode.setValue(newMode),
+                        ),
+                      ),
+                    ),
+                  ),
+                  Expanded(
+                    child: Container(
+                      padding: EdgeInsets.only(left: UnrealTheme.cardMargin, right: UnrealTheme.cardMargin, top: 16),
+                      color: Theme.of(context).colorScheme.surface,
+                      child: PreferenceBuilder(
+                        preference: settings.leftColumnMode,
+                        builder: (context, CameraSettingsLeftColumnMode mode) => switch (mode) {
+                          CameraSettingsLeftColumnMode.camera => _CameraSettingsMainControlsCameraPropertyDisplay(
+                              selectedCamera: selectedCamera,
+                            ),
+                          CameraSettingsLeftColumnMode.focus => _CameraSettingsMainControlsFocusPropertyDisplay(
+                              selectedCamera: selectedCamera,
+                            ),
+                        },
+                      ),
+                    ),
+                  ),
+                ],
+              ),
+            ),
+          ),
+
+          const SizedBox(width: UnrealTheme.sectionMargin),
+
+          // ICFVX depth of field settings
+          Expanded(
+            child: Container(
+              color: Theme.of(context).colorScheme.surface,
+              child: Column(children: [
+                Container(
+                  color: Theme.of(context).colorScheme.surfaceTint,
+                  height: _headerHeight,
+                  child: Center(
+                    child: FakeSelectorBar(AppLocalizations.of(context)!.cameraTabDepthOfFieldPropertiesLabel),
+                  ),
+                ),
+                Expanded(
+                  child: Container(
+                    padding: EdgeInsets.only(left: UnrealTheme.cardMargin, right: UnrealTheme.cardMargin, top: 16),
+                    color: Theme.of(context).colorScheme.surface,
+                    child: _CameraSettingsMainControlsDepthOfFieldPropertyDisplay(
+                      selectedCamera: selectedCamera,
+                    ),
+                  ),
+                ),
+              ]),
+            ),
+          ),
+        ],
+      ),
+    );
+  }
+}
+
+/// Camera property controls in the Camera Settings tab's main controls pane.
+class _CameraSettingsMainControlsCameraPropertyDisplay extends _CineCameraSettingsPropertyDisplay {
+  const _CameraSettingsMainControlsCameraPropertyDisplay({required super.selectedCamera});
+
+  @override
+  State<StatefulWidget> createState() => _CameraSettingsMainControlsCameraPropertyDisplayState();
+}
+
+class _CameraSettingsMainControlsCameraPropertyDisplayState
+    extends _CineCameraSettingsPropertyDisplayState<_CameraSettingsMainControlsCameraPropertyDisplay> {
+  @override
+  Widget build(BuildContext context) {
+    return Column(children: [
+      UnrealActorReferenceViewer(
+        unrealProperties: [
+          if (widget.selectedCamera != null)
+            UnrealProperty(
+              objectPath: widget.selectedCamera!.path,
+              propertyName: 'CameraSettings.ExternalCameraActor',
+            ),
+        ],
+      ),
+      UnrealDeltaSlider(
+        unrealProperties: getCineCameraProperty('CurrentFocalLength'),
+        hardMin: 4,
+        hardMax: 1000,
+      ),
+      UnrealDeltaSlider(
+        unrealProperties: getCineCameraProperty('CurrentAperture'),
+        hardMin: 1.2,
+        hardMax: 22,
+      ),
+    ]);
+  }
+}
+
+/// Focus property controls in the Camera Settings tab's main controls pane.
+class _CameraSettingsMainControlsFocusPropertyDisplay extends _CineCameraSettingsPropertyDisplay {
+  const _CameraSettingsMainControlsFocusPropertyDisplay({required super.selectedCamera});
+
+  @override
+  State<StatefulWidget> createState() => _CameraSettingsMainControlsFocusPropertyDisplayState();
+}
+
+class _CameraSettingsMainControlsFocusPropertyDisplayState
+    extends _CineCameraSettingsPropertyDisplayState<_CameraSettingsMainControlsFocusPropertyDisplay>
+    with GuardedRefreshState {
+  /// Property controller used to listen to changes to the focus method.
+  UnrealPropertyController<String>? _focusMethodController;
+
+  @override
+  void initState() {
+    super.initState();
+
+    _updateFocusMethodTracking();
+  }
+
+  @override
+  void dispose() {
+    _focusMethodController?.dispose();
+
+    super.dispose();
+  }
+
+  @override
+  Widget build(BuildContext context) {
+    final focusMode = _focusMethodController?.getSingleSharedValue().value;
+
+    final List<Widget> propertyWidgets;
+
+    switch (focusMode) {
+      case 'Manual':
+        propertyWidgets = [
+          UnrealDeltaSlider(
+            unrealProperties: getCineCameraProperty('FocusSettings.ManualFocusDistance'),
+            hardMin: 1.5,
+            softMax: 100000,
+            minMaxBehaviour: PropertyMinMaxBehaviour.clamp,
+            exponent: 3.0,
+          ),
+          ..._makeGenericFocusPropertyWidgets(),
+        ];
+
+      case 'Tracking':
+        propertyWidgets = [
+          UnrealActorReferenceViewer(
+            unrealProperties: getCineCameraProperty('FocusSettings.TrackingFocusSettings.ActorToTrack'),
+          ),
+          UnrealDeltaSlider(
+            unrealProperties: getCineCameraProperty('FocusSettings.TrackingFocusSettings.RelativeOffset.X'),
+            overrideName: AppLocalizations.of(context)!.cameraTabFocusRelativeOffsetXLabel,
+            softMin: -10000,
+            softMax: 10000,
+          ),
+          UnrealDeltaSlider(
+            unrealProperties: getCineCameraProperty('FocusSettings.TrackingFocusSettings.RelativeOffset.Y'),
+            overrideName: AppLocalizations.of(context)!.cameraTabFocusRelativeOffsetYLabel,
+            softMin: -10000,
+            softMax: 10000,
+          ),
+          UnrealDeltaSlider(
+            unrealProperties: getCineCameraProperty('FocusSettings.TrackingFocusSettings.RelativeOffset.Z'),
+            overrideName: AppLocalizations.of(context)!.cameraTabFocusRelativeOffsetZLabel,
+            softMin: -10000,
+            softMax: 10000,
+          ),
+          ..._makeGenericFocusPropertyWidgets(),
+        ];
+        break;
+
+      default:
+        propertyWidgets = [];
+        break;
+    }
+
+    return Column(children: [
+      UnrealDropdownSelector(
+        unrealProperties: getCineCameraProperty('FocusSettings.FocusMethod'),
+      ),
+      ...propertyWidgets,
+    ]);
+  }
+
+  @override
+  void onCameraChanged() {
+    _updateFocusMethodTracking();
+  }
+
+  /// Create the list of widgets used in both "enabled" focus modes (i.e. Manual/Tracking).
+  List<Widget> _makeGenericFocusPropertyWidgets() => [
+        UnrealBooleanDropdownSelector(
+          unrealProperties: getCineCameraProperty('FocusSettings.bSmoothFocusChanges'),
+        ),
+        UnrealDeltaSlider(
+          unrealProperties: getCineCameraProperty('FocusSettings.FocusSmoothingInterpSpeed'),
+          hardMin: 0,
+          softMax: 10,
+        ),
+        UnrealDeltaSlider(
+          unrealProperties: getCineCameraProperty('FocusSettings.FocusOffset'),
+          softMin: -1000,
+          softMax: 1000,
+        ),
+      ];
+
+  /// Update our property tracking for the focus method to reflect the currently selected camera.
+  void _updateFocusMethodTracking() {
+    _focusMethodController = UnrealPropertyController(context);
+
+    final UnrealProperty? property = getCineCameraProperty('FocusSettings.FocusMethod').firstOrNull;
+
+    if (property == null) {
+      return;
+    }
+
+    _focusMethodController = UnrealPropertyController(context);
+    _focusMethodController!.addListener(guardedRefresh);
+    _focusMethodController!.trackAllProperties([property]);
+  }
+}
+
+/// Depth of field property controls in the Camera Settings tab's main controls pane.
+class _CameraSettingsMainControlsDepthOfFieldPropertyDisplay extends StatefulWidget {
+  const _CameraSettingsMainControlsDepthOfFieldPropertyDisplay({required this.selectedCamera});
+
+  /// The camera selected for editing.
+  final _CameraSettingsOutlinerEntryData? selectedCamera;
+
+  @override
+  State<StatefulWidget> createState() => _CameraSettingsMainControlsDepthOfFieldPropertyDisplayState();
+}
+
+class _CameraSettingsMainControlsDepthOfFieldPropertyDisplayState
+    extends State<_CameraSettingsMainControlsDepthOfFieldPropertyDisplay> with GuardedRefreshState {
+  /// Property controller used to listen to whether the distance to wall is set automatically.
+  UnrealPropertyController<bool>? _distanceToWallAutoController;
+
+  @override
+  void initState() {
+    super.initState();
+
+    _updateDistanceToWallAutoTracking();
+  }
+
+  @override
+  void didUpdateWidget(covariant _CameraSettingsMainControlsDepthOfFieldPropertyDisplay oldWidget) {
+    super.didUpdateWidget(oldWidget);
+
+    if (oldWidget.selectedCamera?.path != widget.selectedCamera?.path) {
+      _updateDistanceToWallAutoTracking();
+    }
+  }
+
+  @override
+  void dispose() {
+    _distanceToWallAutoController?.dispose();
+
+    super.dispose();
+  }
+
+  @override
+  Widget build(BuildContext context) {
+    final bool bIsDistanceToWallAuto = _distanceToWallAutoController?.getSingleSharedValue().value ?? false;
+
+    final Widget distanceToWallWidget;
+    if (bIsDistanceToWallAuto) {
+      distanceToWallWidget = UnrealDeltaSlider(
+        unrealProperties: [],
+        overrideName: AppLocalizations.of(context)!.cameraTabDepthOfFieldDistanceToWallLabel,
+        valueMessageOverride: AppLocalizations.of(context)!.propertyAutomaticLabel,
+        bShowResetButton: false,
+      );
+    } else {
+      distanceToWallWidget = UnrealDeltaSlider(
+        unrealProperties: getCameraProperty('CameraSettings.CameraDepthOfField.DistanceToWall'),
+        softMin: 0,
+        softMax: 100000,
+        exponent: 3.0,
+      );
+    }
+
+    return Column(children: [
+      UnrealBooleanDropdownSelector(
+        unrealProperties: getCameraProperty('CameraSettings.CameraDepthOfField.bEnableDepthOfFieldCompensation'),
+        overrideName: AppLocalizations.of(context)!.cameraTabDepthOfFieldCompensationLabel,
+      ),
+      UnrealBooleanDropdownSelector(
+        unrealProperties: getCameraProperty('CameraSettings.CameraDepthOfField.bAutomaticallySetDistanceToWall'),
+      ),
+      distanceToWallWidget,
+      UnrealDeltaSlider(
+        unrealProperties: getCameraProperty('CameraSettings.CameraDepthOfField.DistanceToWallOffset'),
+        softMin: -1000,
+        softMax: 1000,
+      ),
+      UnrealDeltaSlider(
+        unrealProperties: getCameraProperty('CameraSettings.CameraDepthOfField.DepthOfFieldGain'),
+      ),
+    ]);
+  }
+
+  /// Get a property by [name] on the selected camera.
+  /// The return value is empty if no camera is selected.
+  List<UnrealProperty> getCameraProperty(String name) {
+    final String? cameraPath = widget.selectedCamera?.path;
+
+    if (cameraPath == null) {
+      return [];
+    }
+
+    return [
+      UnrealProperty(
+        objectPath: cameraPath,
+        propertyName: name,
+      )
+    ];
+  }
+
+  /// Update our property tracking for the focus method to reflect the currently selected camera.
+  void _updateDistanceToWallAutoTracking() {
+    _distanceToWallAutoController = UnrealPropertyController(context);
+
+    final UnrealProperty? property =
+        getCameraProperty('CameraSettings.CameraDepthOfField.bAutomaticallySetDistanceToWall').firstOrNull;
+
+    if (property == null) {
+      return;
+    }
+
+    _distanceToWallAutoController = UnrealPropertyController(context);
+    _distanceToWallAutoController!.addListener(guardedRefresh);
+    _distanceToWallAutoController!.trackAllProperties([property]);
+  }
+}
+
+/// Base class for widgets that display properties in the Camera Settings tab, and may pull these properties from either
+/// the selected camera itself or its associated cine camera.
+abstract class _CineCameraSettingsPropertyDisplay extends StatefulWidget {
+  const _CineCameraSettingsPropertyDisplay({required this.selectedCamera});
+
+  /// The camera selected for editing.
+  final _CameraSettingsOutlinerEntryData? selectedCamera;
+}
+
+/// Base state class for widgets that display properties in the Camera Settings tab, and may pull these properties from
+/// either the selected camera itself or its associated cine camera.
+abstract class _CineCameraSettingsPropertyDisplayState<T extends _CineCameraSettingsPropertyDisplay> extends State<T> {
+  /// The camera that holds the properties. Null if it hasn't been retrieved yet.
+  UnrealObject? _camera;
+
+  /// Property controller used to listen to changes to the focus method.
+  UnrealPropertyController<String>? _cineCameraActorController;
+
+  @override
+  void initState() {
+    super.initState();
+
+    _updateCamera();
+  }
+
+  @override
+  void didUpdateWidget(covariant T oldWidget) {
+    super.didUpdateWidget(oldWidget);
+
+    if (oldWidget.selectedCamera?.path != widget.selectedCamera?.path) {
+      _updateCamera();
+    }
+  }
+
+  @override
+  void dispose() {
+    _cineCameraActorController?.dispose();
+
+    super.dispose();
+  }
+
+  /// Get a property by [name] on the cine camera, or if it isn't present, on the selected camera.
+  /// Returns the property in a list for use with delta widgets.
+  /// The return value is empty if no camera is selected.
+  List<UnrealProperty> getCineCameraProperty(String name) {
+    if (_camera == null) {
+      return [];
+    }
+
+    return [
+      UnrealProperty(
+        objectPath: _camera!.path,
+        propertyName: name,
+      )
+    ];
+  }
+
+  /// Called when the camera whose properties we want to modify changes.
+  @protected
+  void onCameraChanged() {}
+
+  /// Update the path to the [_camera] based on the configured settings.
+  void _updateCamera() async {
+    _cineCameraActorController?.dispose();
+
+    final String? cameraPath = widget.selectedCamera?.path;
+    if (cameraPath == null) {
+      return;
+    }
+
+    _cineCameraActorController = UnrealPropertyController(context);
+    _cineCameraActorController!.addListener(_onCineCameraChanged);
+    _cineCameraActorController!.trackAllProperties([
+      UnrealProperty(
+        objectPath: cameraPath,
+        propertyName: 'CameraSettings.ExternalCameraActor',
+      ),
+    ]);
+  }
+
+  /// Called when the current camera's [cineCameraPath] property value changes.
+  void _onCineCameraChanged() async {
+    if (!mounted) {
+      return;
+    }
+
+    final String? cineCameraPath = _cineCameraActorController?.getSingleSharedValue().value;
+
+    setState(() {
+      _camera = null;
+    });
+    onCameraChanged();
+
+    if (cineCameraPath == null || cineCameraPath.isEmpty) {
+      if (widget.selectedCamera != null && mounted) {
+        // No cine camera set, so use the camera selected in the app
+        setState(() {
+          _camera = UnrealObject(
+            path: widget.selectedCamera!.path,
+            name: widget.selectedCamera!.name,
+          );
+          onCameraChanged();
+        });
+      }
+      return;
+    }
+
+    // Retrieve the camera component path and name of the cine camera actor
+    final connection = Provider.of<EngineConnectionManager>(context, listen: false);
+    final UnrealHttpResponse result = await connection.sendHttpRequest(
+      UnrealHttpRequest(
+        url: '/remote/object/property',
+        verb: 'PUT',
+        body: {
+          'objectPath': cineCameraPath,
+          'propertyName': 'CameraComponent',
+          'access': 'READ_ACCESS',
+        },
+      ),
+    );
+
+    if (!mounted) {
+      return;
+    }
+
+    // Retrieve the cine camera component path
+    if (result.code != HttpResponseCode.ok) {
+      _log.warning('Failed to retrieve cine camera component: ${result.body}');
+      return;
+    }
+
+    final componentPath = result.body['CameraComponent'];
+    if (componentPath is! String) {
+      _log.warning('No camera component found on cine camera actor $cineCameraPath');
+      return;
+    }
+
+    setState(() {
+      _camera = UnrealObject(path: componentPath, name: componentPath.split('.').last);
+      onCameraChanged();
+    });
+  }
+}
+
+/// The outliner panel for the color grading tab, showing cameras that can be targeted.
+class _CameraSettingsOutlinerPanel extends StatefulWidget {
+  const _CameraSettingsOutlinerPanel({
+    required this.entries,
+  });
+
+  /// A flat list of all objects to display in the outliner hierarchy.
+  final List<_CameraSettingsOutlinerEntryData> entries;
+
+  @override
+  State<_CameraSettingsOutlinerPanel> createState() => _CameraSettingsOutlinerPanelState();
+}
+
+class _CameraSettingsOutlinerPanelState extends State<_CameraSettingsOutlinerPanel> {
+  late final TreeViewController _treeController = TreeViewController();
+
+  @override
+  void initState() {
+    super.initState();
+
+    _updateTreeController();
+  }
+
+  @override
+  void didUpdateWidget(covariant _CameraSettingsOutlinerPanel oldWidget) {
+    super.didUpdateWidget(oldWidget);
+
+    _updateTreeController();
+  }
+
+  @override
+  Widget build(BuildContext context) {
+    final tabSettings = Provider.of<CameraSettingsTabSettings>(context, listen: false);
+
+    return Column(
+      children: [
+        CardSmallHeader(title: AppLocalizations.of(context)!.outlinerTitle),
+        Expanded(
+          child: TreeView(
+            padding: EdgeInsets.symmetric(vertical: 4),
+            treeController: _treeController,
+            nodeBuilder: (node, controller) => TransientPreferenceBuilder<String>(
+              preference: tabSettings.selectedCamera,
+              builder: (_, selectedCamera) => _CameraSettingsOutlinerEntry(
+                node: node as TreeViewNode<_CameraSettingsOutlinerEntryData>,
+                controller: _treeController,
+                onTap: () => tabSettings.selectedCamera.setValue(node.data.path),
+                bIsSelected: node.data.path == selectedCamera,
+              ),
+            ),
+          ),
+        ),
+      ],
+    );
+  }
+
+  /// Update the tree controller with current entries and remove old ones.
+  void _updateTreeController() {
+    final Set<String> relevantKeys = {};
+
+    // Add new entries
+    for (final _CameraSettingsOutlinerEntryData entry in widget.entries) {
+      relevantKeys.add(entry.path);
+      if (_treeController.getNode(entry.path) == null) {
+        _treeController.addNode(
+          TreeViewNode<_CameraSettingsOutlinerEntryData>(key: entry.path, data: entry),
+          parentKey: entry.parent?.path,
+        );
+      }
+    }
+
+    // Remove old entries
+    final List<String> allKeys = List.from(_treeController.allKeys);
+    for (final String key in allKeys) {
+      if (!relevantKeys.contains(key)) {
+        _treeController.removeNode(key);
+      }
+    }
+  }
+}
+
+/// Types of object that can be shown in the color grading tab's outliner panel.
+enum _CameraSettingsOutlinerEntryType {
+  nDisplayConfig,
+  icvfxCamera,
+}
+
+/// Widget representing an object listed in the camera settings outliner panel.
+class _CameraSettingsOutlinerEntry extends StatelessWidget {
+  const _CameraSettingsOutlinerEntry({
+    required this.node,
+    required this.controller,
+    required this.onTap,
+    required this.bIsSelected,
+    Key? key,
+  }) : super(key: key);
+
+  static const Map<_CameraSettingsOutlinerEntryType, String> _iconsByType = {
+    _CameraSettingsOutlinerEntryType.icvfxCamera: 'packages/epic_common/assets/icons/ndisplay_camera.svg',
+    _CameraSettingsOutlinerEntryType.nDisplayConfig: 'packages/epic_common/assets/icons/ndisplay.svg',
+  };
+
+  /// The node in the tree view for this item.
+  final TreeViewNode<_CameraSettingsOutlinerEntryData> node;
+
+  /// Tree view controller for the containing tree.
+  final TreeViewController controller;
+
+  /// Callback for when this item is tapped.
+  final Function() onTap;
+
+  /// Whether the item is currently selected.
+  final bool bIsSelected;
+
+  /// Whether this can be selected as a target object.
+  bool get bCanBeSelected => node.data.type != _CameraSettingsOutlinerEntryType.nDisplayConfig;
+
+  /// The path of the icon to display.
+  String? get iconPath => _iconsByType[node.data.type];
+
+  @override
+  Widget build(BuildContext context) {
+    late final CardListTileExpansionState expansionState;
+    if (node.children.length > 0) {
+      expansionState = controller.isNodeExpanded(node.key)
+          ? CardListTileExpansionState.expanded
+          : CardListTileExpansionState.collapsed;
+    } else {
+      expansionState = CardListTileExpansionState.none;
+    }
+
+    Widget tileBuilder(context, bIsEnabled) => CardListTile(
+          bDeEmphasize: !bIsEnabled,
+          bIsSelected: bCanBeSelected ? bIsSelected : false,
+          title: node.data.name,
+          iconPath: iconPath,
+          onTap: bCanBeSelected ? _onTap : null,
+          indentation: node.indentation,
+          expansionState: expansionState,
+        );
+
+    return tileBuilder(context, true);
+  }
+
+  /// Called when the main body of the entry is tapped.
+  void _onTap() {
+    if (!bCanBeSelected) {
+      return;
+    }
+
+    onTap();
+  }
+}
+
+/// Data about an entry in the camera settings tab's outliner panel.
+class _CameraSettingsOutlinerEntryData {
+  const _CameraSettingsOutlinerEntryData({
+    required this.path,
+    required this.name,
+    required this.type,
+    this.parent,
+  });
+
+  /// The path of the actor or component this represents.
+  final String path;
+
+  /// The name to show in the outliner panel.
+  final String name;
+
+  /// The type of actor/component represented by this entry.
+  final _CameraSettingsOutlinerEntryType type;
+
+  /// The object this is nested under in the Outliner panel, if any.
+  final _CameraSettingsOutlinerEntryData? parent;
+}
diff --git a/Engine/Extras/VirtualProduction/EpicStageApp/lib/widgets/screens/main/tabs/color_grading_tab.dart b/Engine/Extras/VirtualProduction/EpicStageApp/lib/widgets/screens/main/tabs/color_grading_tab.dart
--- a/Engine/Extras/VirtualProduction/EpicStageApp/lib/widgets/screens/main/tabs/color_grading_tab.dart
+++ b/Engine/Extras/VirtualProduction/EpicStageApp/lib/widgets/screens/main/tabs/color_grading_tab.dart
@@ -22,6 +22,7 @@ import '../../../../models/unreal_actor_manager.dart';
 import '../../../../models/unreal_transaction_manager.dart';
 import '../../../../models/unreal_types.dart';
 import '../../../../utilities/constants.dart';
+import '../../../../utilities/net_utilities.dart';
 import '../../../../utilities/unreal_utilities.dart';
 import '../../../elements/delta_slider.dart';
 import '../../../elements/dropdown_text.dart';
@@ -462,43 +463,17 @@ class _ColorGradingTabState extends State<ColorGradingTab> {
       return;
     }
 
-    // Try to get the per-node color grading settings from each component as if it's an ICVFX camera. If it succeeds,
-    // we can treat it as one and add all its targets to the list.
-    final List<UnrealHttpRequest> describeRequests = [];
-    for (final dynamic componentPath in componentPaths) {
-      if (componentPath is! String) {
-        continue;
-      }
-
-      describeRequests.add(UnrealHttpRequest(url: '/remote/object/describe', verb: 'PUT', body: {
-        'objectPath': componentPath,
-        'access': 'READ_ACCESS',
-      }));
-    }
-
-    final UnrealHttpResponse batchDescribeResponse = await _connectionManager.sendBatchedHttpRequest(describeRequests);
-    if (batchDescribeResponse.code != 200 || !mounted) {
-      return;
-    }
-
-    // Find which of the described responses are actually cameras
-    final List<String> cameraPaths = [];
-    for (int componentIndex = 0; componentIndex < componentPaths.length; ++componentIndex) {
-      final UnrealHttpResponse? describeResponse = batchDescribeResponse.body[componentIndex];
-      if (describeResponse == null || describeResponse.code != 200) {
-        continue;
-      }
-
-      if (describeResponse.body?['Class'] == '/Script/DisplayCluster.DisplayClusterICVFXCameraComponent') {
-        cameraPaths.add(componentPaths[componentIndex]);
-      }
-    }
+    final List<UnrealObject> cameras = await getComponentsOfType(
+      connectionManager: _connectionManager,
+      componentPaths: componentPaths.map((e) => e.toString()).toList(growable: false),
+      classPath: '/Script/DisplayCluster.DisplayClusterICVFXCameraComponent',
+    );
 
     // Now that we know which are cameras, query them all for their color grading targets and add them to the list
     final List<UnrealHttpRequest> propertyRequests = [];
-    for (final String cameraPath in cameraPaths) {
+    for (final UnrealObject camera in cameras) {
       propertyRequests.add(UnrealHttpRequest(url: '/remote/object/property', verb: 'PUT', body: {
-        'objectPath': cameraPath,
+        'objectPath': camera.path,
         'propertyName': 'CameraSettings.PerNodeColorGrading',
         'access': 'READ_ACCESS',
       }));
@@ -509,16 +484,16 @@ class _ColorGradingTabState extends State<ColorGradingTab> {
       return;
     }
 
-    for (int cameraIndex = 0; cameraIndex < cameraPaths.length; ++cameraIndex) {
+    for (int cameraIndex = 0; cameraIndex < cameras.length; ++cameraIndex) {
       final UnrealHttpResponse? propertyResponse = batchPropertyResponse.body[cameraIndex];
-      _addTargetsForICVFXCameraResponse(cameraPaths[cameraIndex], propertyResponse, parent);
+      _addTargetsForICVFXCameraResponse(cameras[cameraIndex], propertyResponse, parent);
     }
   }
 
   /// Given a response from querying an ICVFX camera's list of per-node color grading settings, add all of its color
   /// grading targets to the target list.
   void _addTargetsForICVFXCameraResponse(
-      String cameraPath, UnrealHttpResponse? response, _ColorGradingObjectEntryData parent) {
+      UnrealObject camera, UnrealHttpResponse? response, _ColorGradingObjectEntryData parent) {
     if (response?.code != 200) {
       return;
     }
@@ -527,7 +502,7 @@ class _ColorGradingTabState extends State<ColorGradingTab> {
 
     // All nodes color grading is always available
     final allNodesRootProperty = UnrealProperty(
-      objectPath: cameraPath,
+      objectPath: camera.path,
       propertyName: 'CameraSettings.AllNodesColorGrading',
     );
 
@@ -552,7 +527,7 @@ class _ColorGradingTabState extends State<ColorGradingTab> {
       }
 
       final targetRootProperty = UnrealProperty(
-        objectPath: cameraPath,
+        objectPath: camera.path,
         propertyName: 'CameraSettings.PerNodeColorGrading[$entryIndex]',
       );
 
@@ -565,11 +540,11 @@ class _ColorGradingTabState extends State<ColorGradingTab> {
       ));
     }
 
-    final int lastDotIndex = cameraPath.lastIndexOf('.');
-    final String cameraName = cameraPath.substring(lastDotIndex + 1);
+    final int lastDotIndex = camera.path.lastIndexOf('.');
+    final String cameraName = camera.path.substring(lastDotIndex + 1);
 
     _entries.add(_ColorGradingObjectEntryData(
-      path: cameraPath,
+      path: camera.path,
       name: cameraName,
       type: _ColorGradingObjectEntryType.icvfxCamera,
       targetListTitle: AppLocalizations.of(context)!.colorGradingOutlinerTargetListTitlePerNode,
diff --git a/Engine/Extras/VirtualProduction/EpicStageApp/lib/widgets/screens/main/tabs/main_screen_tabs.dart b/Engine/Extras/VirtualProduction/EpicStageApp/lib/widgets/screens/main/tabs/main_screen_tabs.dart
--- a/Engine/Extras/VirtualProduction/EpicStageApp/lib/widgets/screens/main/tabs/main_screen_tabs.dart
+++ b/Engine/Extras/VirtualProduction/EpicStageApp/lib/widgets/screens/main/tabs/main_screen_tabs.dart
@@ -4,6 +4,7 @@ import 'dart:collection';
 
 import 'package:flutter/material.dart';
 
+import 'camera_settings_tab.dart';
 import 'color_grading_tab.dart';
 import 'details_tab.dart';
 import 'stage_map_tab.dart';
@@ -23,6 +24,11 @@ class MainScreenTabs {
       iconPath: DetailsTab.iconPath,
       createTabContents: (context) => const DetailsTab(),
     ),
+    MainScreenTabConfig(
+      getTitle: CameraSettingsTab.getTitle,
+      iconPath: CameraSettingsTab.iconPath,
+      createTabContents: (context) => const CameraSettingsTab(),
+    ),
     MainScreenTabConfig(
       getTitle: ColorGradingTab.getTitle,
       iconPath: ColorGradingTab.iconPath,
diff --git a/Engine/Extras/VirtualProduction/EpicStageApp/pubspec.yaml b/Engine/Extras/VirtualProduction/EpicStageApp/pubspec.yaml
--- a/Engine/Extras/VirtualProduction/EpicStageApp/pubspec.yaml
+++ b/Engine/Extras/VirtualProduction/EpicStageApp/pubspec.yaml
@@ -18,7 +18,7 @@ publish_to: 'none' # Remove this line if you wish to publish to pub.dev
 version: 1.0.3+4
 
 environment:
-  sdk: ">=2.17.0 <3.0.0"
+  sdk: ">=3.0.6 <4.0.0"
   flutter: "^3.16.3"
 
 # Dependencies specify other packages that your package needs in order to work.
