[UnrealStage] Fix flickering camera settings when Cine Camera Actor is set

diff --git a/Engine/Extras/VirtualProduction/EpicStageApp/lib/widgets/elements/actor_reference_viewer.dart b/Engine/Extras/VirtualProduction/EpicStageApp/lib/widgets/elements/actor_reference_viewer.dart
--- a/Engine/Extras/VirtualProduction/EpicStageApp/lib/widgets/elements/actor_reference_viewer.dart
+++ b/Engine/Extras/VirtualProduction/EpicStageApp/lib/widgets/elements/actor_reference_viewer.dart
@@ -95,10 +95,6 @@ class _UnrealActorDisplayState extends State<UnrealActorReferenceViewer>
 
   @override
   void handleOnPropertiesChanged() {
-    setState(() {
-      _actorName = '';
-    });
-
     _updateDisplayName();
   }
 
diff --git a/Engine/Extras/VirtualProduction/EpicStageApp/lib/widgets/screens/main/tabs/camera_settings_tab.dart b/Engine/Extras/VirtualProduction/EpicStageApp/lib/widgets/screens/main/tabs/camera_settings_tab.dart
--- a/Engine/Extras/VirtualProduction/EpicStageApp/lib/widgets/screens/main/tabs/camera_settings_tab.dart
+++ b/Engine/Extras/VirtualProduction/EpicStageApp/lib/widgets/screens/main/tabs/camera_settings_tab.dart
@@ -654,8 +654,24 @@ abstract class _CineCameraSettingsPropertyDisplay extends StatefulWidget {
 /// either the selected camera itself or its associated cine camera.
 abstract class _CineCameraSettingsPropertyDisplayState<T extends _CineCameraSettingsPropertyDisplay> extends State<T> {
   /// The camera that holds the properties. Null if it hasn't been retrieved yet.
+  /// Use [camera] to access this value, which will automatically call [onCameraChanged] when appropriate.
   UnrealObject? _camera;
 
+  /// The camera that holds the properties. Null if it hasn't been retrieved yet.
+  UnrealObject? get camera => _camera;
+
+  set camera(UnrealObject? value) {
+    if (!mounted || value?.path == _camera?.path) {
+      return;
+    }
+
+    setState(() {
+      _camera = value;
+    });
+
+    onCameraChanged();
+  }
+
   /// Property controller used to listen to changes to the focus method.
   UnrealPropertyController<String>? _cineCameraActorController;
 
@@ -686,13 +702,13 @@ abstract class _CineCameraSettingsPropertyDisplayState<T extends _CineCameraSett
   /// Returns the property in a list for use with delta widgets.
   /// The return value is empty if no camera is selected.
   List<UnrealProperty> getCineCameraProperty(String name) {
-    if (_camera == null) {
+    if (camera == null) {
       return [];
     }
 
     return [
       UnrealProperty(
-        objectPath: _camera!.path,
+        objectPath: camera!.path,
         propertyName: name,
       )
     ];
@@ -702,7 +718,7 @@ abstract class _CineCameraSettingsPropertyDisplayState<T extends _CineCameraSett
   @protected
   void onCameraChanged() {}
 
-  /// Update the path to the [_camera] based on the configured settings.
+  /// Update the path to the [camera] based on the configured settings.
   void _updateCamera() async {
     _cineCameraActorController?.dispose();
 
@@ -729,21 +745,13 @@ abstract class _CineCameraSettingsPropertyDisplayState<T extends _CineCameraSett
 
     final String? cineCameraPath = _cineCameraActorController?.getSingleSharedValue().value;
 
-    setState(() {
-      _camera = null;
-    });
-    onCameraChanged();
-
     if (cineCameraPath == null || cineCameraPath.isEmpty) {
       if (widget.selectedCamera != null && mounted) {
         // No cine camera set, so use the camera selected in the app
-        setState(() {
-          _camera = UnrealObject(
-            path: widget.selectedCamera!.path,
-            name: widget.selectedCamera!.name,
-          );
-          onCameraChanged();
-        });
+        camera = UnrealObject(
+          path: widget.selectedCamera!.path,
+          name: widget.selectedCamera!.name,
+        );
       }
       return;
     }
@@ -769,19 +777,18 @@ abstract class _CineCameraSettingsPropertyDisplayState<T extends _CineCameraSett
     // Retrieve the cine camera component path
     if (result.code != HttpResponseCode.ok) {
       _log.warning('Failed to retrieve cine camera component: ${result.body}');
+      camera = null;
       return;
     }
 
     final componentPath = result.body['CameraComponent'];
     if (componentPath is! String) {
       _log.warning('No camera component found on cine camera actor $cineCameraPath');
+      camera = null;
       return;
     }
 
-    setState(() {
-      _camera = UnrealObject(path: componentPath, name: componentPath.split('.').last);
-      onCameraChanged();
-    });
+    camera = UnrealObject(path: componentPath, name: componentPath.split('.').last);
   }
 }
 
