[LiveLinkVCAM] Add final app icon and name

diff --git a/Engine/Extras/VirtualProduction/FlutterLiveLinkVCAM/IOS/Runner/Assets.xcassets/LaunchImage.imageset/Contents.json b/Engine/Extras/VirtualProduction/FlutterLiveLinkVCAM/IOS/Runner/Assets.xcassets/LaunchImage.imageset/Contents.json
--- a/Engine/Extras/VirtualProduction/FlutterLiveLinkVCAM/IOS/Runner/Assets.xcassets/LaunchImage.imageset/Contents.json
+++ b/Engine/Extras/VirtualProduction/FlutterLiveLinkVCAM/IOS/Runner/Assets.xcassets/LaunchImage.imageset/Contents.json
@@ -1,23 +1,23 @@
 {
   "images" : [
     {
-      "idiom" : "universal",
       "filename" : "LaunchImage.png",
+      "idiom" : "universal",
       "scale" : "1x"
     },
     {
-      "idiom" : "universal",
       "filename" : "[REDACTED_EMAIL]",
+      "idiom" : "universal",
       "scale" : "2x"
     },
     {
-      "idiom" : "universal",
       "filename" : "[REDACTED_EMAIL]",
+      "idiom" : "universal",
       "scale" : "3x"
     }
   ],
   "info" : {
-    "version" : 1,
-    "author" : "xcode"
+    "author" : "xcode",
+    "version" : 1
   }
 }
diff --git a/Engine/Extras/VirtualProduction/FlutterLiveLinkVCAM/android/app/src/main/AndroidManifest.xml b/Engine/Extras/VirtualProduction/FlutterLiveLinkVCAM/android/app/src/main/AndroidManifest.xml
--- a/Engine/Extras/VirtualProduction/FlutterLiveLinkVCAM/android/app/src/main/AndroidManifest.xml
+++ b/Engine/Extras/VirtualProduction/FlutterLiveLinkVCAM/android/app/src/main/AndroidManifest.xml
@@ -33,7 +33,7 @@
     <uses-permission android:name="android.permission.BLUETOOTH_CONNECT" />
 
     <application
-        android:label="live_link_vcam"
+        android:label="Unreal VCam"
         android:name="${applicationName}"
         android:icon="@mipmap/ic_launcher"
         android:enableOnBackInvokedCallback="true">
diff --git a/Engine/Extras/VirtualProduction/FlutterLiveLinkVCAM/android/app/src/main/res/drawable/launch_background.xml b/Engine/Extras/VirtualProduction/FlutterLiveLinkVCAM/android/app/src/main/res/drawable/launch_background.xml
--- a/Engine/Extras/VirtualProduction/FlutterLiveLinkVCAM/android/app/src/main/res/drawable/launch_background.xml
+++ b/Engine/Extras/VirtualProduction/FlutterLiveLinkVCAM/android/app/src/main/res/drawable/launch_background.xml
@@ -1,12 +1,12 @@
 <?xml version="1.0" encoding="utf-8"?>
 <!-- Modify this file to customize your launch splash screen -->
 <layer-list xmlns:android="http://schemas.android.com/apk/res/android">
-    <item android:drawable="@android:color/white" />
+    <item android:drawable="@android:color/black" />
 
     <!-- You can insert your own image assets here -->
-    <!-- <item>
+    <item>
         <bitmap
             android:gravity="center"
-            android:src="@mipmap/launch_image" />
-    </item> -->
+            android:src="@mipmap/ic_launcher" />
+    </item>
 </layer-list>
diff --git a/Engine/Extras/VirtualProduction/FlutterLiveLinkVCAM/android/app/src/main/res/values/styles.xml b/Engine/Extras/VirtualProduction/FlutterLiveLinkVCAM/android/app/src/main/res/values/styles.xml
--- a/Engine/Extras/VirtualProduction/FlutterLiveLinkVCAM/android/app/src/main/res/values/styles.xml
+++ b/Engine/Extras/VirtualProduction/FlutterLiveLinkVCAM/android/app/src/main/res/values/styles.xml
@@ -1,7 +1,7 @@
 <?xml version="1.0" encoding="utf-8"?>
 <resources>
     <!-- Theme applied to the Android Window while the process is starting when the OS's Dark Mode setting is off -->
-    <style name="LaunchTheme" parent="@android:style/Theme.Light.NoTitleBar">
+    <style name="LaunchTheme" parent="@android:style/Theme.Black.NoTitleBar">
         <!-- Show a splash screen on the activity. Automatically removed when
              the Flutter engine draws its first frame -->
         <item name="android:windowBackground">@drawable/launch_background</item>
@@ -12,7 +12,7 @@
          running.
 
          This Theme is only used starting with V2 of Flutter's Android embedding. -->
-    <style name="NormalTheme" parent="@android:style/Theme.Light.NoTitleBar">
+    <style name="NormalTheme" parent="@android:style/Theme.Black.NoTitleBar">
         <item name="android:windowBackground">?android:colorBackground</item>
     </style>
 </resources>
diff --git a/Engine/Extras/VirtualProduction/FlutterLiveLinkVCAM/pubspec.yaml b/Engine/Extras/VirtualProduction/FlutterLiveLinkVCAM/pubspec.yaml
--- a/Engine/Extras/VirtualProduction/FlutterLiveLinkVCAM/pubspec.yaml
+++ b/Engine/Extras/VirtualProduction/FlutterLiveLinkVCAM/pubspec.yaml
@@ -38,6 +38,7 @@ dev_dependencies:
   ffigen: ^10.0.0
   pigeon: ^14.0.1
   url_launcher: ^6.2.2
+  flutter_launcher_icons: ^0.13.1
 
   # The "flutter_lints" package below contains a set of recommended lints to
   # encourage good coding practices. The lint set provided by the package is
@@ -78,3 +79,11 @@ arb_generator:
   csv_settings:
     description_index: 1
     base_index: 2
+
+# Configuration for automatic icon generation.
+# Run `dart run flutter_launcher_icons` from a terminal to generate these.
+flutter_launcher_icons:
+  android: true
+  ios: true
+  remove_alpha_ios: true
+  image_path: "assets/app_icon/VCam_Icon_HD.png"
\ No newline at end of file
