[UnrealStage] Refactor floating tabs and add draggable bar to trackpad

diff --git a/Engine/Extras/VirtualProduction/EpicStageApp/lib/main.dart b/Engine/Extras/VirtualProduction/EpicStageApp/lib/main.dart
--- a/Engine/Extras/VirtualProduction/EpicStageApp/lib/main.dart
+++ b/Engine/Extras/VirtualProduction/EpicStageApp/lib/main.dart
@@ -28,7 +28,6 @@ import 'models/settings/selected_actor_settings.dart';
 import 'models/settings/stage_map_settings.dart';
 import 'models/unreal_actor_creator.dart';
 import 'models/unreal_actor_manager.dart';
-import 'models/unreal_dockable_tab_manager.dart';
 import 'models/unreal_property_manager.dart';
 import 'models/unreal_transaction_manager.dart';
 import 'routes.dart';
@@ -192,10 +191,6 @@ class _EpicStageAppState extends State<EpicStageApp> with TickerProviderStateMix
           create: (context) => UnrealActorCreator(context),
           dispose: (context, value) => value.dispose(),
         ),
-        Provider<UnrealDockableTabManager>(
-          create: (context) => UnrealDockableTabManager(context),
-          dispose: (context, value) => value.dispose(),
-        )
       ],
       child: MaterialApp(
         theme: UnrealTheme.makeThemeData(),
diff --git a/Engine/Extras/VirtualProduction/EpicStageApp/lib/models/settings/floating_tab_settings.dart b/Engine/Extras/VirtualProduction/EpicStageApp/lib/models/settings/floating_tab_settings.dart
deleted file mode 100644
diff --git a/Engine/Extras/VirtualProduction/EpicStageApp/lib/models/settings/floating_window_settings.dart b/Engine/Extras/VirtualProduction/EpicStageApp/lib/models/settings/floating_window_settings.dart
new file mode 100644
--- /dev/null
+++ b/Engine/Extras/VirtualProduction/EpicStageApp/lib/models/settings/floating_window_settings.dart
@@ -0,0 +1,39 @@
+// Copyright Epic Games, Inc. All Rights Reserved.
+
+import 'package:epic_common/preferences.dart';
+import 'package:streaming_shared_preferences/streaming_shared_preferences.dart';
+
+/// Which side of the screen a tab is on.
+enum FloatingWindowSide {
+  left,
+  right,
+}
+
+/// Generic holder for user settings used to configure floating windows.
+/// Takes [preference] bundle containing user settings and a [prefix] string used to disambiguate per-window settings
+/// [FloatingWindowSettings]
+class FloatingWindowSettings {
+  FloatingWindowSettings(PreferencesBundle preferences, String prefix)
+      : side = preferences.persistent.getEnum<FloatingWindowSide>(
+          '$prefix.side',
+          defaultValue: FloatingWindowSide.right,
+          enumValues: FloatingWindowSide.values,
+        ),
+        yAxis = preferences.persistent.getDouble(
+          '$prefix.yAxis',
+          defaultValue: 1.0,
+        ),
+        bIsDocked = preferences.persistent.getBool(
+          '$prefix.bIsDocked',
+          defaultValue: true,
+        );
+
+  /// Which side of the screen the window is on.
+  Preference<FloatingWindowSide> side;
+
+  /// The vertical position at which the tab should sit (as a fraction of its container height).
+  Preference<double> yAxis;
+
+  /// Whether the window is docked off-screen.
+  Preference<bool> bIsDocked;
+}
diff --git a/Engine/Extras/VirtualProduction/EpicStageApp/lib/models/unreal_dockable_tab_manager.dart b/Engine/Extras/VirtualProduction/EpicStageApp/lib/models/unreal_dockable_tab_manager.dart
deleted file mode 100644
diff --git a/Engine/Extras/VirtualProduction/EpicStageApp/lib/widgets/elements/floating_map_preview.dart b/Engine/Extras/VirtualProduction/EpicStageApp/lib/widgets/elements/floating_map_preview.dart
--- a/Engine/Extras/VirtualProduction/EpicStageApp/lib/widgets/elements/floating_map_preview.dart
+++ b/Engine/Extras/VirtualProduction/EpicStageApp/lib/widgets/elements/floating_map_preview.dart
@@ -2,54 +2,24 @@
 
 import 'package:epic_common/widgets.dart';
 import 'package:flutter/material.dart';
-import 'package:provider/provider.dart';
 
-import 'floating_tab_base.dart';
+import 'floating_window.dart';
 import 'lightcard_map.dart';
 
-/// A floating preview of the stage map that the user can drag to any corner of the screen and dock to the side.
-class FloatingMapPreview extends StatefulWidget {
-  const FloatingMapPreview({Key? key}) : super(key: key);
-
-  static GlobalKey map = GlobalKey(debugLabel: 'floating_map');
-
-  @override
-  State<FloatingMapPreview> createState() => _FloatingMapPreviewState();
-}
-
-class _FloatingMapPreviewState extends State<FloatingMapPreview> with SingleTickerProviderStateMixin, FloatingTabBase {
-  @override
-  void initState() {
-    super.initState();
-
-    tabController = manager.mapTabController;
-
-    initStationaryController(manager.trackpadTabController);
-
-    tabController.ticker = createTicker(onTick);
-
-    /// init dragged tab relative position.
-    tabRelativePosition =
-        manager.getTabRelativePosition(tabController.tabPosition, stationaryTabController.tabPosition);
-  }
+class FloatingMapPreview extends StatelessWidget {
+  const FloatingMapPreview({super.key});
 
   @override
   Widget build(BuildContext context) {
-    return ChangeNotifierProvider.value(
-      value: tabController,
-      builder: (context, child) {
-        return FloatingTabBody(
-          AssetIcon(path: 'packages/epic_common/assets/icons/map.svg', size: 24),
-          StageMap(bIsControllable: false, bShowOnlySelectedPins: true, bForceShowFullMap: true, pinSize: 32),
-          bShouldUseStackLayout: true,
-          onTap: onTap,
-          onPanStart: onStartDrag,
-          onPanUpdate: updateDrag,
-          onPanCancel: endDrag,
-          onPanEnd: endDrag,
-          children: [],
-        );
-      },
+    return FloatingWindow(
+      settingsPrefix: 'floatingMapSettings',
+      icon: AssetIcon(path: 'packages/epic_common/assets/icons/map.svg'),
+      builder: (context) => StageMap(
+        bIsControllable: false,
+        bShowOnlySelectedPins: true,
+        bForceShowFullMap: true,
+        pinSize: 32,
+      ),
     );
   }
 }
diff --git a/Engine/Extras/VirtualProduction/EpicStageApp/lib/widgets/elements/floating_tab_base.dart b/Engine/Extras/VirtualProduction/EpicStageApp/lib/widgets/elements/floating_tab_base.dart
deleted file mode 100644
diff --git a/Engine/Extras/VirtualProduction/EpicStageApp/lib/widgets/elements/floating_trackpad.dart b/Engine/Extras/VirtualProduction/EpicStageApp/lib/widgets/elements/floating_trackpad.dart
--- a/Engine/Extras/VirtualProduction/EpicStageApp/lib/widgets/elements/floating_trackpad.dart
+++ b/Engine/Extras/VirtualProduction/EpicStageApp/lib/widgets/elements/floating_trackpad.dart
@@ -1,119 +1,75 @@
 // Copyright Epic Games, Inc. All Rights Reserved.
 
-import 'dart:async';
+import 'dart:ui';
 
-import 'package:epic_common/preferences.dart';
-import 'package:epic_common/theme.dart';
-import 'package:epic_common/widgets.dart';
 import 'package:flutter/material.dart';
-import 'package:provider/provider.dart';
+import 'package:flutter_gen/gen_l10n/app_localizations.dart';
 
-import '../../models/settings/floating_tab_settings.dart';
-import '../../models/settings/selected_actor_settings.dart';
-import 'floating_tab_base.dart';
+import 'floating_window.dart';
 import 'lightcard_trackpad.dart';
 
-/// A floating preview of the stage map that the user can drag to any corner of the screen and dock to the side.
-class FloatingTrackpad extends StatefulWidget {
-  const FloatingTrackpad({Key? key}) : super(key: key);
+class FloatingTrackpad extends StatelessWidget {
+  const FloatingTrackpad({super.key});
 
-  static GlobalKey trackpad = GlobalKey(debugLabel: 'floating_trackpad');
-
-  @override
-  State<FloatingTrackpad> createState() => _FloatingTrackpadState();
-}
-
-class _FloatingTrackpadState extends State<FloatingTrackpad> with SingleTickerProviderStateMixin, FloatingTabBase {
-  /// Whether trackpad is locked for LC positioning input or unlocked for trackpad positioning drag.
-  bool _bIsTrackpadLocked = false;
-  late final SelectedActorSettings _actorSettings;
-
-  /// Instance of user base floating tab settings for trackpad.
-  late final FloatingTabSettings _settings;
-
-  /// Stream subscription for when the [tabController.dockState] changes.
-  late final StreamSubscription _subscription;
-
-  @override
-  void initState() {
-    super.initState();
-
-    tabController = manager.trackpadTabController;
-
-    initStationaryController(manager.mapTabController);
-
-    tabController.ticker = createTicker(onTick);
-
-    _actorSettings = Provider.of<SelectedActorSettings>(context, listen: false);
-
-    _settings = FloatingTabSettings(PreferencesBundle.of(context), trackpadPreferencePrefix);
-
-    /// init dragged tab relative position.
-    tabRelativePosition =
-        manager.getTabRelativePosition(tabController.tabPosition, stationaryTabController.tabPosition);
-
-    /// This helps us listen for changes to our dockState, if this tab is ever undocked while [_bIsTrackpadLocked] is
-    /// true, all drag and tap gestures would be ignored. to prevent this we want to know when the dock state changes
-    /// and set [_bIsTrackpadLocked] to false.
-    _subscription = _settings.bIsTabDocked.listen((value) => _handleOnDockStateChanged);
-  }
-
-  /// handle on dock state changed.
-  void _handleOnDockStateChanged(bool value) {
-    if (value && _bIsTrackpadLocked) {
-      setState(() => _bIsTrackpadLocked = false);
-    }
-  }
-
-  @override
-  void dispose() {
-    _subscription.cancel();
-    super.dispose();
-  }
+  /// The height of the draggable area at the top of the trackpad.
+  static const double _dragBarHeight = 44;
 
   @override
   Widget build(BuildContext context) {
-    tabAvoidanceOnFirstTimeRun();
-    return ChangeNotifierProvider.value(
-      value: tabController,
-      builder: (context, child) {
-        return FloatingTabBody(
-          Icon(Icons.games_outlined),
-          LightCardTrackPad(),
-          bIsDraggable: !_bIsTrackpadLocked && _actorSettings.selectedActors.getValue().isNotEmpty,
-          bCanInteract: _bIsTrackpadLocked,
-          onTap: onTap,
-          onPanStart: onStartDrag,
-          onPanUpdate: _bIsTrackpadLocked ? null : updateDrag,
-          onPanCancel: endDrag,
-          onPanEnd: endDrag,
-          children: [
-            // Lock/unlock button which determines whether or not the trackpad is open modifying the position of a
-            // selected lightcard.
-            Visibility(
-              visible: tabController.dockState == DockState.undocked,
-              child: Padding(
-                padding: const EdgeInsets.all(8.0),
-                child: SizedBox(
-                  width: 40,
-                  height: 40,
-                  child: EpicIconButton(
-                    bIsToggledOn: _bIsTrackpadLocked,
-                    activeBackgroundColor: UnrealColors.gray18.withOpacity(.1),
-                    backgroundShape: BoxShape.circle,
-                    iconPath: _bIsTrackpadLocked
-                        ? 'packages/epic_common/assets/icons/lock_locked.svg'
-                        : 'packages/epic_common/assets/icons/lock_unlocked.svg',
-                    onPressed: () => setState(() {
-                      _bIsTrackpadLocked = !_bIsTrackpadLocked;
-                    }),
+    final size = FloatingWindow.getDefaultSize(context) + const Offset(0, _dragBarHeight);
+
+    return FloatingWindow(
+      settingsPrefix: 'floatingTrackpad',
+      icon: Icon(Icons.games_outlined),
+      builder: (context) => Stack(
+        children: [
+          Positioned(
+            top: _dragBarHeight - 1,
+            bottom: 0,
+            left: 0,
+            right: 0,
+            child: BackdropFilter(
+              filter: ImageFilter.blur(sigmaX: 3, sigmaY: 3),
+              child: Container(
+                color: Colors.black54,
+                child: Center(
+                  child: Icon(
+                    Icons.touch_app_rounded,
+                    size: 100,
+                    color: Colors.white12,
                   ),
                 ),
               ),
             ),
-          ],
-        );
-      },
+          ),
+          Positioned.fill(
+            child: LightCardTrackpad(),
+          ),
+          Container(
+            height: _dragBarHeight,
+            padding: const EdgeInsets.symmetric(horizontal: 4, vertical: 10),
+            alignment: Alignment.centerLeft,
+            color: Theme.of(context).colorScheme.surfaceTint,
+            child: Row(
+              children: [
+                Icon(
+                  Icons.drag_indicator,
+                  color: Colors.white30,
+                ),
+                const SizedBox(width: 4),
+                Expanded(
+                  child: Text(
+                    AppLocalizations.of(context)!.trackpadTitle.toUpperCase(),
+                    style: Theme.of(context).textTheme.titleLarge,
+                  ),
+                ),
+              ],
+            ),
+          ),
+        ],
+      ),
+      draggableInsets: EdgeInsets.only(bottom: size.height - _dragBarHeight),
+      size: size,
     );
   }
 }
diff --git a/Engine/Extras/VirtualProduction/EpicStageApp/lib/widgets/elements/floating_window.dart b/Engine/Extras/VirtualProduction/EpicStageApp/lib/widgets/elements/floating_window.dart
new file mode 100644
--- /dev/null
+++ b/Engine/Extras/VirtualProduction/EpicStageApp/lib/widgets/elements/floating_window.dart
@@ -0,0 +1,745 @@
+// Copyright Epic Games, Inc. All Rights Reserved.
+
+import 'dart:math' as math;
+
+import 'package:epic_common/preferences.dart';
+import 'package:epic_common/widgets.dart';
+import 'package:flutter/material.dart';
+import 'package:flutter/physics.dart';
+import 'package:flutter/scheduler.dart';
+
+import '../../models/settings/floating_window_settings.dart';
+
+/// Sides of the screen on which the window can be docked.
+enum _DockSide {
+  left,
+  right,
+}
+
+/// Possible states for the floating window in terms of docking.
+enum _DockState {
+  /// The window is on-screen and no dock tabs are visible.
+  undocked,
+
+  /// The window has entered the docking range and is revealing its tab.
+  docking,
+
+  /// The window has revealed its docking tab. It may not actually be off-screen yet.
+  docked,
+
+  /// The window has exited the docking range and is hiding its tab.
+  undocking,
+}
+
+/// How far to inset the window from the edge of the screen.
+const double _edgeInset = 8.0;
+
+/// The size of the tab shown when the window is docked.
+const _dockTabSize = Size(32, 80);
+
+/// The smallest dimension of the container will be multiplied by this value to determine the height of the window.
+const double _windowSizeScale = 0.25;
+
+/// The height of the window will be multiplied by this value to determine its width.
+const double _windowAspectRatio = 16.0 / 9.0;
+
+/// Builds the contents of a floating window in the given build [context].
+typedef FloatingWindowContentBuilder = Widget Function(BuildContext context);
+
+/// A floating window that the user can drag to any corner of the screen and dock to the side.
+class FloatingWindow extends StatefulWidget {
+  const FloatingWindow({
+    Key? key,
+    required this.settingsPrefix,
+    required this.icon,
+    required this.builder,
+    this.size = null,
+    this.draggableInsets = EdgeInsets.zero,
+  }) : super(key: key);
+
+  /// The prefix used to look up the tab's specific user settings.
+  final String settingsPrefix;
+
+  /// The icon to display on the tab when docked.
+  final Widget icon;
+
+  /// The function to build the tab's contents.
+  final FloatingWindowContentBuilder builder;
+
+  /// The size of the window.
+  final Size? size;
+
+  /// When the window is undocked, the area that the user can drag will be inset by this much.
+  final EdgeInsetsGeometry draggableInsets;
+
+  /// Calculate the default size for a floating window in this [context].
+  static Size getDefaultSize(BuildContext context) {
+    final Size screenSize = MediaQuery.of(context).size;
+    final double windowHeight = math.min(screenSize.width, screenSize.height) * _windowSizeScale;
+
+    return Size(windowHeight * _windowAspectRatio, windowHeight);
+  }
+
+  @override
+  State<FloatingWindow> createState() => _FloatingWindowState();
+}
+
+class _FloatingWindowState extends State<FloatingWindow> with SingleTickerProviderStateMixin {
+  /// Physics settings for the spring simulation when the user is dragging the window.
+  static const _draggedSpringDescription = SpringDescription(mass: 40, stiffness: 1, damping: 1.2);
+
+  /// Physics settings for the spring simulation when the user releases the window.
+  static const _releasedSpringDescription = SpringDescription(mass: 40, stiffness: 1, damping: 1);
+
+  /// Physics settings for the spring simulation when the window is being docked.
+  static const _dockedSpringDescription = SpringDescription(mass: 60, stiffness: 1, damping: 1);
+
+  /// When this portion of the window's width is offscreen, dock it.
+  static const double _dockThreshold = 1 / 3;
+
+  /// Which side the window was last docked on.
+  _DockSide _dockSide = _DockSide.left;
+
+  /// State of the window window's docking.
+  _DockState _dockState = _DockState.undocked;
+
+  /// The position that the window's top-left corner should move towards with spring motion.
+  Offset _targetWindowPosition = Offset.zero;
+
+  /// The current position of the window's top-left corner.
+  Offset _windowPosition = Offset.zero;
+
+  /// The current velocity of the window.
+  Offset _windowVelocity = Offset.zero;
+
+  /// The user's last drag position in global space.
+  Offset _lastDragPosition = Offset.zero;
+
+  /// The last known size of the containing render box.
+  Size _containerSize = Size.zero;
+
+  /// The elapsed ticker time in seconds when the springs were last updated.
+  double _lastSpringTime = 0;
+
+  /// The last value of [_lastSpringTime] when the springs' target positions were moved.
+  double _lastSpringMoveTime = 0;
+
+  /// Spring controlling the X axis.
+  SpringSimulation? _xSpring;
+
+  /// Spring controlling the Y axis.
+  SpringSimulation? _ySpring;
+
+  /// Ticker used to update springs.
+  late final Ticker _ticker;
+
+  /// Settings for this window map.
+  late final FloatingWindowSettings _settings;
+
+  /// The size of the window to display given the containing screen's size.
+  Size get _windowSize => widget.size ?? FloatingWindow.getDefaultSize(context);
+
+  /// True if the window is currently docked or in the process of docking.
+  bool get _bIsDockedOrDocking => _dockState == _DockState.docking || _dockState == _DockState.docked;
+
+  /// True if the window is currently undocked or in the process of undocking.
+  bool get _bIsUndockedOrUndocking => _dockState == _DockState.undocking || _dockState == _DockState.undocked;
+
+  /// True if the window is not docked and offscreen.
+  bool get _bIsWindowVisible {
+    if (_dockState != _DockState.docked) {
+      return true;
+    }
+
+    // Check if the window is actually within visible bounds (e.g. still moving to docked position or being dragged out
+    // of it).
+    return ((_windowPosition.dx + _windowSize.width - _dockTabSize.width) > 0 &&
+        (_windowPosition.dx + _dockTabSize.width) < _containerSize.width);
+  }
+
+  @override
+  Widget build(BuildContext context) {
+    return LayoutBuilder(builder: (BuildContext context, BoxConstraints constraints) {
+      final Size newSize = Size(constraints.maxWidth, constraints.maxHeight);
+
+      // Handle first size/resize
+      if (newSize != _containerSize) {
+        final Size oldSize = _containerSize;
+        _containerSize = newSize;
+
+        if (oldSize == Size.zero) {
+          _initPosition();
+        } else {
+          WidgetsBinding.instance.addPostFrameCallback((_) => _onContainerResized(oldSize));
+        }
+      }
+
+      // Determine the offset of the main window window, which shifts to the side to reveal dock tabs when docked.
+      double dockOffset = 0;
+      if (_bIsDockedOrDocking) {
+        switch (_dockSide) {
+          case _DockSide.left:
+            dockOffset = -_dockTabSize.width;
+            break;
+
+          case _DockSide.right:
+            dockOffset = _dockTabSize.width;
+            break;
+        }
+      }
+
+      return Stack(
+        clipBehavior: Clip.hardEdge,
+        children: [
+          // Position a box containing the window, dock tabs, and gesture detector
+          Positioned(
+            top: _windowPosition.dy,
+            left: _windowPosition.dx,
+            child: SizedBox(
+              width: _windowSize.width,
+              height: _windowSize.height,
+              child: Stack(
+                clipBehavior: Clip.none,
+                children: [
+                  // Tab shown on the left when docked on the right
+                  if (_dockState != _DockState.undocked && _dockSide == _DockSide.right)
+                    Positioned(
+                      key: Key('LeftDockTab'),
+                      left: 0,
+                      top: (_windowSize.height - _dockTabSize.height) / 2,
+                      child: _DockTab(
+                        side: _DockSide.right,
+                        icon: widget.icon,
+                      ),
+                    ),
+
+                  // Tab shown on the right when docked on the left
+                  if (_dockState != _DockState.undocked && _dockSide == _DockSide.left)
+                    Positioned(
+                      key: Key('RightDockTab'),
+                      left: _windowSize.width - _dockTabSize.width,
+                      top: (_windowSize.height - _dockTabSize.height) / 2,
+                      child: _DockTab(
+                        side: _DockSide.left,
+                        icon: widget.icon,
+                      ),
+                    ),
+
+                  // Window window, which shifts to reveal a tab when docked and hides when offscreen
+                  AnimatedPositioned(
+                    key: Key('WindowMover'),
+                    left: dockOffset,
+                    duration: const Duration(milliseconds: 200),
+                    child: _bIsWindowVisible
+                        ? _WindowBody(
+                            key: Key('Window'),
+                            size: _windowSize,
+                            bShadowVisible: _bIsUndockedOrUndocking,
+                            child: widget.builder(context),
+                          )
+                        : SizedBox(),
+                    onEnd: _onDockTabAnimationEnd,
+                  ),
+
+                  // Gesture detector, which matches the window rectangle when undocked and the revealed tab when docked
+                  Positioned(
+                    key: Key('GestureDetector'),
+
+                    // When docked to left, offset to line up with the tab on the right side of the window
+                    left: (_dockState == _DockState.docked && _dockSide == _DockSide.left)
+                        ? (_windowSize.width - _dockTabSize.width)
+                        : 0,
+
+                    // When docked, offset to line up with the top of the tab
+                    top: _dockState == _DockState.docked ? ((_windowSize.height - _dockTabSize.height) / 2) : 0,
+
+                    child: SizedBox(
+                      // When docked, match size to the tab
+                      width: _dockState == _DockState.docked ? _dockTabSize.width : _windowSize.width,
+                      height: _dockState == _DockState.docked ? _dockTabSize.height : _windowSize.height,
+
+                      child: Padding(
+                        padding: _dockState == _DockState.docked ? EdgeInsets.zero : widget.draggableInsets,
+                        child: GestureDetector(
+                          behavior: HitTestBehavior.opaque,
+                          onTap: _onTap,
+                          onPanUpdate: _updateDrag,
+                          onPanEnd: (details) => _endDrag(details.velocity),
+                          onPanCancel: () => _endDrag(Velocity.zero),
+                        ),
+                      ),
+                    ),
+                  ),
+                ],
+              ),
+            ),
+          ),
+        ],
+      );
+    });
+  }
+
+  @override
+  void initState() {
+    super.initState();
+
+    _ticker = createTicker(_onTick);
+    _settings = FloatingWindowSettings(
+      PreferencesBundle.of(context),
+      widget.settingsPrefix,
+    );
+  }
+
+  @override
+  void dispose() {
+    _ticker.dispose();
+
+    super.dispose();
+  }
+
+  /// Initialize the position of the window once we a valid size to reference.
+  /// This should NOT call setState, as it's called directly from the first build().
+  void _initPosition() {
+    final restoredYPosition = _settings.yAxis.getValue() * _containerSize.height;
+    final FloatingWindowSide mapSide = _settings.side.getValue();
+
+    if (_settings.bIsDocked.getValue()) {
+      switch (mapSide) {
+        case FloatingWindowSide.left:
+          _dockSide = _DockSide.left;
+          break;
+
+        case FloatingWindowSide.right:
+          _dockSide = _DockSide.right;
+          break;
+      }
+
+      _dockState = _DockState.docked;
+      _targetWindowPosition = _windowPosition = _getDockedWindowPosition(dockY: restoredYPosition);
+    } else {
+      // Make an approximate position for the window, which we'll then snap to the nearest corner
+      late final double tempXPosition;
+      switch (mapSide) {
+        case FloatingWindowSide.left:
+          tempXPosition = 0;
+          break;
+
+        case FloatingWindowSide.right:
+          tempXPosition = _containerSize.width - _windowSize.width;
+          break;
+      }
+
+      final tempPosition = Offset(tempXPosition, restoredYPosition);
+
+      _dockState = _DockState.undocked;
+      _targetWindowPosition = _windowPosition = _getNearestWindowCornerPosition(tempPosition);
+    }
+  }
+
+  /// Called when the containing widget is resized and we need to readjust accordingly.
+  void _onContainerResized(Size oldContainerSize) {
+    if (_containerSize.height == 0 || _containerSize.width == 0) {
+      return;
+    }
+
+    final Offset scaledPosition = Offset(
+      _targetWindowPosition.dx * (_containerSize.width / oldContainerSize.width),
+      _targetWindowPosition.dy * (_containerSize.height / oldContainerSize.height),
+    );
+
+    if (_bIsDockedOrDocking) {
+      // Stay docked and snap to the new side
+      setState(() {
+        _targetWindowPosition = _windowPosition = _getDockedWindowPosition(dockY: scaledPosition.dy);
+      });
+    } else {
+      // Move smoothly to the new position
+      _moveWindowTargetToNearestCorner(scaledPosition);
+    }
+  }
+
+  /// Handle a tap on the window.
+  void _onTap() {
+    if (_bIsUndockedOrUndocking) {
+      return;
+    }
+
+    // Currently docked, so undock by moving to the nearest corner.
+    _moveWindowTargetToNearestCorner(_targetWindowPosition);
+  }
+
+  /// Update the dragged window position.
+  void _updateDrag(DragUpdateDetails details) {
+    _lastDragPosition = details.globalPosition;
+
+    _moveWindowTarget(
+      _targetWindowPosition + Offset(details.delta.dx, details.delta.dy),
+      springDescription: _draggedSpringDescription,
+    );
+
+    final _DockSide? newDockSide = _checkDockSideAtPosition(_targetWindowPosition);
+    if (newDockSide != null) {
+      // If this will dock, show the tab
+      _dockWindow(newDockSide, moveOffscreen: false);
+    } else {
+      // No longer docked, so move to undocking state
+      _undockWindow();
+    }
+  }
+
+  /// Stop dragging the window and let it come to rest.
+  void _endDrag(Velocity velocity) {
+    final RenderBox? renderBox = context.findRenderObject() as RenderBox?;
+    if (renderBox == null) {
+      _moveWindowTarget(
+        Offset(_edgeInset, _edgeInset),
+        springDescription: _releasedSpringDescription,
+      );
+      return;
+    }
+
+    // Check ahead from the window's position so the user can "fling" the window off the screen
+    const dockCheckAheadSeconds = 0.1;
+    final Offset dockCheckPosition = _windowPosition + velocity.pixelsPerSecond * dockCheckAheadSeconds;
+
+    // If near the edges, dock the window
+    final _DockSide? newDockSide = _checkDockSideAtPosition(dockCheckPosition);
+    if (newDockSide != null) {
+      _dockWindow(newDockSide, dockY: dockCheckPosition.dy);
+      return;
+    }
+
+    // Otherwise, move to the nearest corner. In this case, use the user's drag position instead of the window
+    // position since flinging feels more accurate that way.
+    const cornerCheckAheadSeconds = 0.15;
+    final Offset cornerCheckPosition = _lastDragPosition + velocity.pixelsPerSecond * cornerCheckAheadSeconds;
+    final Offset finalLocalPosition = renderBox.globalToLocal(cornerCheckPosition);
+
+    _moveWindowTargetToNearestCorner(finalLocalPosition);
+  }
+
+  /// Update spring window position using springs.
+  void _onTick(Duration elapsed) {
+    if (_xSpring == null || _ySpring == null) {
+      return;
+    }
+
+    _lastSpringTime = elapsed.inMilliseconds / 1000.0;
+
+    final double springTime = _lastSpringTime - _lastSpringMoveTime;
+
+    setState(() {
+      _windowVelocity = Offset(_xSpring!.dx(springTime), _ySpring!.dx(springTime));
+      _windowPosition = Offset(_xSpring!.x(springTime), _ySpring!.x(springTime));
+    });
+
+    // If both springs are at rest, stop ticking to save performance
+    if (_xSpring!.isDone(springTime) && _ySpring!.isDone(springTime)) {
+      _ticker.stop();
+      _lastSpringMoveTime = 0;
+      _lastSpringTime = 0;
+    }
+  }
+
+  /// Called when the animation to reveal/hide a dock tab completes.
+  void _onDockTabAnimationEnd() {
+    switch (_dockState) {
+      case _DockState.undocking:
+        setState(() {
+          _dockState = _DockState.undocked;
+        });
+        break;
+
+      case _DockState.docking:
+        setState(() {
+          _dockState = _DockState.docked;
+        });
+        break;
+
+      default:
+        break;
+    }
+  }
+
+  /// Move the target position for the window to a new location and adjust springs accordingly.
+  void _moveWindowTarget(Offset newTarget, {required SpringDescription springDescription}) {
+    // Clamp both axes in a valid range
+    _targetWindowPosition = Offset(
+      newTarget.dx.clamp(-_windowSize.width + _dockTabSize.width, _containerSize.width - _dockTabSize.width),
+      _getClampedWindowYPosition(newTarget.dy),
+    );
+
+    _xSpring = SpringSimulation(
+      springDescription,
+      _windowPosition.dx,
+      _targetWindowPosition.dx,
+      _windowVelocity.dx,
+    );
+
+    _ySpring = SpringSimulation(
+      springDescription,
+      _windowPosition.dy,
+      _targetWindowPosition.dy,
+      _windowVelocity.dy,
+    );
+
+    _lastSpringMoveTime = _lastSpringTime;
+
+    if (!_ticker.isActive) {
+      _ticker.start();
+    }
+
+    _updateSettingsWindowMapY(newTarget.dy);
+  }
+
+  /// Move the window target to the corner nearest to a local [position].
+  void _moveWindowTargetToNearestCorner(Offset position) {
+    _moveWindowTarget(
+      _getNearestWindowCornerPosition(position),
+      springDescription: _releasedSpringDescription,
+    );
+
+    _settings.side.setValue(
+      (_targetWindowPosition.dx < _containerSize.width / 2) ? FloatingWindowSide.left : FloatingWindowSide.right,
+    );
+
+    _undockWindow();
+  }
+
+  /// Reveal the dock tab and move to the docked state to the given [side].
+  /// If [moveOffscreen] is true, move the window off-screen, optionally using the Y position specified by [dockY].
+  void _dockWindow(_DockSide side, {bool moveOffscreen = true, double? dockY}) {
+    dockY = dockY ?? _windowPosition.dy;
+
+    if (_bIsUndockedOrUndocking) {
+      // Trigger the docking animation
+      setState(() {
+        _dockSide = side;
+        _dockState = _DockState.docking;
+      });
+
+      // Update user settings
+      switch (side) {
+        case _DockSide.left:
+          _settings.side.setValue(FloatingWindowSide.left);
+          break;
+
+        case _DockSide.right:
+          _settings.side.setValue(FloatingWindowSide.right);
+          break;
+
+        default:
+          break;
+      }
+
+      _settings.bIsDocked.setValue(true);
+      _updateSettingsWindowMapY(dockY);
+    }
+
+    if (!moveOffscreen) {
+      return;
+    }
+
+    _moveWindowTarget(
+      _getDockedWindowPosition(dockY: dockY),
+      springDescription: _dockedSpringDescription,
+    );
+  }
+
+  /// Hide the docked tab and move to the undocking state.
+  void _undockWindow() {
+    if (_bIsDockedOrDocking) {
+      setState(() {
+        _dockState = _DockState.undocking;
+      });
+
+      _settings.bIsDocked.setValue(false);
+    }
+  }
+
+  /// Check whether the window should be docked if it's at the given X position.
+  /// Returns the side to dock to, or null if it shouldn't dock.
+  _DockSide? _checkDockSideAtPosition(Offset position) {
+    if (position.dx + ((1 - _dockThreshold) * _windowSize.width) > _containerSize.width) {
+      return _DockSide.right;
+    }
+
+    if (position.dx + (_dockThreshold * _windowSize.width) < 0) {
+      return _DockSide.left;
+    }
+
+    return null;
+  }
+
+  /// Update the saved Y position of the docked window.
+  void _updateSettingsWindowMapY(double dockY) {
+    _settings.yAxis.setValue(dockY / _containerSize.height);
+  }
+
+  /// Clamp a Y position for the window to a valid range.
+  double _getClampedWindowYPosition(double y) {
+    final double minY = _edgeInset;
+    final double maxY = _containerSize.height - _windowSize.height - _edgeInset;
+
+    return math.max(math.min(y, maxY), minY);
+  }
+
+  /// Get the position of the window when moved to the nearest corner to a local [position].
+  Offset _getNearestWindowCornerPosition(Offset position) {
+    late final double targetX;
+    if (position.dx < _containerSize.width / 2) {
+      targetX = _edgeInset;
+    } else {
+      targetX = _containerSize.width - _windowSize.width - _edgeInset;
+    }
+
+    late final double targetY;
+    if (position.dy < _containerSize.height / 2) {
+      targetY = _edgeInset;
+    } else {
+      targetY = _containerSize.height - _windowSize.height - _edgeInset;
+    }
+
+    return Offset(targetX, targetY);
+  }
+
+  /// Get the docked position of the window when its Y position is [dockY].
+  Offset _getDockedWindowPosition({required double dockY}) {
+    // Extra amount to move just to make sure it's fully off-screen so we can stop rendering the window
+    const double offsetFudge = 0.5;
+
+    // Move the window to the appropriate side
+    late final double targetX;
+    switch (_dockSide) {
+      case _DockSide.left:
+        targetX = -_windowSize.width + _dockTabSize.width - offsetFudge;
+        break;
+
+      case _DockSide.right:
+        targetX = _containerSize.width - _dockTabSize.width + offsetFudge;
+        break;
+    }
+
+    return Offset(targetX, _getClampedWindowYPosition(dockY));
+  }
+}
+
+/// The main body of the window which moves around within the larger widget.
+class _WindowBody extends StatelessWidget {
+  const _WindowBody({
+    Key? key,
+    required this.size,
+    required this.bShadowVisible,
+    required this.child,
+  }) : super(key: key);
+
+  /// The size of the window's body.
+  final Size size;
+
+  /// If true, add a shadow effect to the window's body.
+  final bool bShadowVisible;
+
+  /// The widget contained within the window's body.
+  final Widget child;
+
+  @override
+  Widget build(BuildContext context) {
+    final borderRadius = BorderRadius.circular(6);
+
+    return SizedBox(
+      width: size.width,
+      height: size.height,
+      child: Stack(
+        children: [
+          Positioned.fill(
+            child: AnimatedContainer(
+              clipBehavior: Clip.antiAlias,
+              duration: Duration(milliseconds: 400),
+              decoration: BoxDecoration(
+                borderRadius: borderRadius,
+                boxShadow: [
+                  BoxShadow(
+                    color: Color(0xff000000).withOpacity(bShadowVisible ? 0.4 : 0),
+                    spreadRadius: 2,
+                    blurRadius: 5,
+                  ),
+                ],
+              ),
+              child: child,
+            ),
+          ),
+          IgnorePointer(
+            child: Container(
+              decoration: BoxDecoration(
+                borderRadius: borderRadius,
+                border: Border.all(
+                  color: Color(0x50ffffff),
+                  width: 0.5,
+                ),
+              ),
+            ),
+          ),
+        ],
+      ),
+    );
+  }
+}
+
+/// The tab shown when the window is docked.
+class _DockTab extends StatelessWidget {
+  const _DockTab({
+    required this.side,
+    required this.icon,
+  });
+
+  /// Which side the window is docked to.
+  final _DockSide side;
+
+  /// The icon to display on the tab when docked.
+  final Widget icon;
+
+  @override
+  Widget build(BuildContext context) {
+    const Radius borderRadius = Radius.circular(8);
+    final bool bIsFacingLeft = side == _DockSide.right;
+
+    return Container(
+      width: _dockTabSize.width,
+      height: _dockTabSize.height,
+      decoration: BoxDecoration(
+        color: Theme.of(context).colorScheme.surfaceTint,
+        borderRadius: bIsFacingLeft
+            ? BorderRadius.only(
+                topLeft: borderRadius,
+                bottomLeft: borderRadius,
+              )
+            : BorderRadius.only(
+                topRight: borderRadius,
+                bottomRight: borderRadius,
+              ),
+        boxShadow: [
+          BoxShadow(
+            color: Color(0x40000000),
+            offset: Offset(bIsFacingLeft ? -4 : 4, 4),
+            blurRadius: 4,
+          ),
+        ],
+      ),
+      child: Column(mainAxisAlignment: MainAxisAlignment.center, children: [
+        SizedBox.square(
+          dimension: 24,
+          child: icon,
+        ),
+        SizedBox(
+          height: 6,
+        ),
+        AssetIcon(
+          path: bIsFacingLeft
+              ? 'packages/epic_common/assets/icons/chevron_left.svg'
+              : 'packages/epic_common/assets/icons/chevron_right.svg',
+          size: 24,
+        ),
+      ]),
+    );
+  }
+}
diff --git a/Engine/Extras/VirtualProduction/EpicStageApp/lib/widgets/elements/lightcard_trackpad.dart b/Engine/Extras/VirtualProduction/EpicStageApp/lib/widgets/elements/lightcard_trackpad.dart
--- a/Engine/Extras/VirtualProduction/EpicStageApp/lib/widgets/elements/lightcard_trackpad.dart
+++ b/Engine/Extras/VirtualProduction/EpicStageApp/lib/widgets/elements/lightcard_trackpad.dart
@@ -6,65 +6,49 @@ import 'package:flutter/material.dart';
 import 'package:provider/provider.dart';
 
 import '../../models/actor_data/light_card_actor_data.dart';
+import '../../models/property_modify_operations.dart';
 import '../../models/settings/delta_widget_settings.dart';
-import '../../models/unreal_property_controller.dart';
 import '../../models/settings/selected_actor_settings.dart';
-
-import '../../models/property_modify_operations.dart';
 import '../../models/unreal_actor_manager.dart';
+import '../../models/unreal_property_controller.dart';
 import '../../models/unreal_types.dart';
 import '../../utilities/constants.dart';
+import 'unreal_widget_base.dart';
 
-enum _TrackpadOperation { negative, positive }
-
-/// Trackpad widget to moving and positioning LCs, works like an actual trackpad.
-class LightCardTrackPad extends StatefulWidget {
-  const LightCardTrackPad({super.key});
+/// Turns the area into a virtual trackpad, allowing the user to control the selected actors' latitude and longitude
+/// by dragging on it.
+class LightCardTrackpad extends StatefulWidget {
+  const LightCardTrackpad({super.key});
 
   @override
-  State<LightCardTrackPad> createState() => _LightCardTrackPadState();
+  State<LightCardTrackpad> createState() => _LightCardTrackpadState();
 }
 
-class _LightCardTrackPadState extends State<LightCardTrackPad> {
+class _LightCardTrackpadState extends State<LightCardTrackpad> {
+  /// List of actor classes that can be controlled by the trackpad.
+  static const List<String> _controllableClasses = [lightCardClassName, ...colorCorrectWindowClassNames];
+
+  /// Base change rate of latitude/longitude per logical pixel of change in the user's pointer position.
+  static const double baseDeltaMultiplier = 0.4;
+
+  /// Maximum value of latitude at which the "north" pole sits. "South" pole is assumed to be at the negation of this.
+  static const _latitudeMax = 90;
+
+  /// Operation used for all property modifications.
+  static const _propertyOperation = const AddOperation();
+
   late final UnrealActorManager _actorManager;
   late final SelectedActorSettings _selectedActorSettings;
-
   late final DeltaWidgetSettings _deltaSettings;
 
-  /// Property Controller for latitude.
+  /// Property controller for latitude.
   late final UnrealPropertyController<double> _latitudeController;
 
   /// Property controller for longitude.
   late final UnrealPropertyController<double> _longitudeController;
 
-  /// List containing whether or not the trackpad is reversed for each selected actor.
-  List<bool> _listBIsTrackpadReversed = [];
-
-  /// Gets a list of selected actor paths.
-  List<String> _getSelectedActorPaths() {
-    final List<String> positionedActorClasses = [lightCardClassName];
-    positionedActorClasses.addAll(colorCorrectWindowClassNames);
-
-    final List<String> nonUVActorPaths = [];
-
-    for (final String actorPath in _selectedActorSettings.selectedActors.getValue()) {
-      final UnrealObject? actor = _actorManager.getActorAtPath(actorPath);
-      if (!(actor?.isAny(positionedActorClasses) ?? false)) {
-        // Actor is not a class with position properties, so leave it out entirely
-        continue;
-      }
-
-      final LightCardActorData? lightCardActorData = actor!.getPerClassData<LightCardActorData>();
-      if (lightCardActorData?.bIsUV == false) {
-        nonUVActorPaths.add(actorPath);
-      }
-    }
-
-    return nonUVActorPaths;
-  }
-
-  /// Getter for list of selected actors path.
-  List<String> get _paths => _getSelectedActorPaths();
+  /// Set of property indices for actors whose trackpad controls are reversed for the current drag operation.
+  Set<int> _reversedControls = {};
 
   /// stream subscriptions for watching changes to selected actors.
   StreamSubscription? _selectedActorsSubscription;
@@ -78,10 +62,12 @@ class _LightCardTrackPadState extends State<LightCardTrackPad> {
     _latitudeController = UnrealPropertyController(context, bShouldInitTransaction: false);
     _longitudeController = UnrealPropertyController(context, bShouldInitTransaction: false);
 
-    _handleUpdateTrackedProperties();
+    _updateTrackedProperties();
 
     _selectedActorsSubscription = _selectedActorSettings.selectedActors.listen((event) {
-      _handleUpdateTrackedProperties();
+      if (mounted) {
+        _updateTrackedProperties();
+      }
     });
 
     super.initState();
@@ -98,207 +84,157 @@ class _LightCardTrackPadState extends State<LightCardTrackPad> {
   @override
   Widget build(BuildContext context) {
     return GestureDetector(
+      behavior: HitTestBehavior.opaque,
       onPanUpdate: _onPanUpdate,
       onPanEnd: _onPanEnd,
       onPanCancel: _onPanCancel,
-      onPanStart: _onPanStart,
-      child: Container(
-        decoration: BoxDecoration(
-          borderRadius: BorderRadius.circular(6),
-          color: Colors.black,
-        ),
-      ),
+      onTap: () => print('trackpad'),
     );
   }
 
-  /// Handle vertical drag/pan across the trackpad widget providing Drag [details] containing positions on the x
-  /// and y axis.
-  void _onVerticalDrag(DragUpdateDetails details) {
-    List<double> summations = [];
+  /// Gets a list of selected actor paths, filtering out any actors that can't be controlled by the trackpad.
+  List<String> _getSelectedActorPaths() {
+    final List<String> validActorPaths = [];
 
-    double yAxis = details.delta.dy * _deltaSettings.sensitivity.getValue();
+    for (final String actorPath in _selectedActorSettings.selectedActors.getValue()) {
+      final UnrealObject? actor = _actorManager.getActorAtPath(actorPath);
 
-    /// get actors latitude (property value) at [index].
-    double _getPropertyValue(int index) {
-      if (_latitudeController.makeValueDataList().isNotEmpty) {
-        return _latitudeController.makeValueDataList()[index]?.value ?? 0;
+      if (actor == null) {
+        continue;
       }
-      return 0;
-    }
-
-    /// increment the actors longitude at [index] with [value].
-    void _incrementLongitude(double value, int index) {
-      _modifyProperty(value, index, _longitudeController, PropertyMinMaxBehaviour.loop);
-    }
 
-    /// increment the actors longitude at [index] by 180° only when the actors latitude value at [index] is at 90° or
-    /// -90°.
-    void _incrementLongitudeByOneEighty(int index) {
-      if (_getPropertyValue(index) == 90 || _getPropertyValue(index) == -90) {
-        _incrementLongitude(180, index);
+      if (!actor.isAny(_controllableClasses)) {
+        // Actor is not a class with position properties, so leave it out entirely
+        continue;
       }
-    }
-
-    /// handle bounce back for when we hit the north pole or its opposite side, specifying a negative/positive
-    /// [operation] indicated whether the actors latitude is at 90° or -90° respectively, [summation]
-    /// which denotes the to be position on the Y-axis derived from adding the current y-position and the incoming
-    /// position from [details.delta.dy], and [index] which is the property index to be modified.
-    void _handleBounceBack(_TrackpadOperation operation, double summation, int index) {
-      // remainder from after we drive the latitude to 90° or -90°.
-      double overflow = operation == _TrackpadOperation.positive ? summation - 90 : -90 - summation;
 
-      // Drive the latitude to the max of 90° or -90°
-      _modifyProperty(yAxis - overflow, index, _latitudeController);
-
-      _incrementLongitudeByOneEighty(index);
-
-      // Drive the latitude in the reverse direction with the [overflow].
-      double tempValue = operation == _TrackpadOperation.positive ? -overflow * 2 : overflow * 2;
-
-      _modifyProperty(tempValue, index, _latitudeController);
+      final LightCardActorData? lightCardActorData = actor.getPerClassData<LightCardActorData>();
+      if (lightCardActorData?.bIsUV == true) {
+        // This is an LC with UV coordinates, so it has no lat/long to control
+        continue;
+      }
 
-      // specify that we want our trackpad to now be reversed.
-      _listBIsTrackpadReversed[index] = true;
+      validActorPaths.add(actorPath);
     }
 
-    /// handle case for where the trackpad is marked as reversed.
-    void _handleReverseTrackpad(double summation, int index) {
-      bool bIsTrackpadReversed = _listBIsTrackpadReversed[index];
+    return validActorPaths;
+  }
 
-      if (bIsTrackpadReversed) {
-        double currentPosition = _getPropertyValue(index);
+  /// Update tracked properties for both latitude and longitude controllers.
+  void _updateTrackedProperties() {
+    _reversedControls.clear();
 
-        if (summation < 90) {
-          _modifyProperty(-yAxis, index, _latitudeController);
+    _longitudeController.trackAllProperties(_getPositionalProperties('Longitude'));
+    _latitudeController.trackAllProperties(_getPositionalProperties('Latitude'));
+  }
 
-          if (currentPosition == 90) {
-            _incrementLongitudeByOneEighty(index);
-            _listBIsTrackpadReversed[index] = false;
-          }
-        }
+  /// Called whenever a pan gesture's position updates based on new user input.
+  void _onPanUpdate(DragUpdateDetails details) {
+    // Transaction is global, so this will be shared with the latitude controller
+    _longitudeController.beginTransaction();
 
-        if (currentPosition == -90) {
-          _incrementLongitudeByOneEighty(index);
-          _modifyProperty(yAxis, index, _latitudeController);
+    final Offset delta = details.delta * baseDeltaMultiplier * _deltaSettings.sensitivity.getValue();
 
-          _listBIsTrackpadReversed[index] = false;
-        }
+    // Longitude is a simple loop operation at the min/max values, so we can update it directly
+    _longitudeController.modifyProperties(
+      _propertyOperation,
+      values: List.generate(
+        _longitudeController.properties.length,
+        (_) => delta.dx,
+        growable: false,
+      ),
+      minMaxBehaviour: PropertyMinMaxBehaviour.loop,
+    );
 
+    // Latitude requires special handling at the min/max values. When we reach either pole, the longitude rotates by
+    // 180 degrees and changes in latitude delta reverses direction to continue moving smoothly.
+    // For example, if the user is dragging up to move north, once they reach the north pole, they expect to continue
+    // circling the globe in the same direction, i.e. starting to move south on the opposite longitudinal side of the
+    // globe.
+    _latitudeController.properties.indexed.forEach((propertyPair) {
+      final int propertyIndex = propertyPair.$1;
+      final WidgetControlledUnrealProperty<double>? property = propertyPair.$2;
+
+      final double? currentValue = property?.value;
+      if (currentValue == null) {
         return;
       }
 
-      _modifyProperty(yAxis, index, _latitudeController);
-    }
+      final bool bIsReversed = _reversedControls.contains(propertyIndex);
 
-    /// get (summations) to be positions of all actors on the y-axis by adding the value of [details.delta.dy] to an
-    /// actors current position.
-    for (int index = 0; index < _paths.length; index++) {
-      double latitude = 0;
+      double deltaY = delta.dy;
 
-      /// check to make sure all selected actors have a valid position on the y-axis and that the length of paths
-      /// matches the length of data values for the actors.
-      if (_latitudeController.makeValueDataList().isNotEmpty &&
-          _paths.length == _latitudeController.makeValueDataList().length) {
-        latitude = _latitudeController.makeValueDataList()[index]?.value ?? 0;
+      // Latitude moves the opposite direction of where the user would expect on a trackpad, so negate it by default
+      if (!bIsReversed) {
+        deltaY = -deltaY;
       }
 
-      /// sum of incoming dy [yAxis] and the current Y position [latitude] of LC'(s).
-      double summation = latitude + yAxis;
+      double newValue = currentValue + deltaY;
+      final double absNewValue = newValue.abs();
 
-      summations.add(summation);
-    }
+      if (absNewValue < _latitudeMax) {
+        // Just move the distance specified
+        _latitudeController.modifyProperty(_propertyOperation, propertyIndex, value: deltaY);
+        return;
+      }
 
-    /// handle when we hit either the negative or positive south pole.
-    for (int index = 0; index < summations.length; index++) {
-      double sum = summations[index];
-      if (sum > 90) {
-        _handleBounceBack(_TrackpadOperation.positive, sum, index);
-      } else if (sum < -90) {
-        _handleBounceBack(_TrackpadOperation.negative, sum, index);
+      // Move to the pole
+      final double distanceToPole = (_latitudeMax - absNewValue) * newValue.sign;
+      _latitudeController.modifyProperty(_propertyOperation, propertyIndex, value: distanceToPole);
+
+      // Flip longitude
+      _longitudeController.modifyProperty(
+        const AddOperation(),
+        propertyIndex,
+        value: 180,
+        minMaxBehaviour: PropertyMinMaxBehaviour.loop,
+      );
+
+      // Move the remaining distance away from the pole
+      final double remainingDistance = (_latitudeMax - absNewValue) * newValue.sign;
+      _latitudeController.modifyProperty(_propertyOperation, propertyIndex, value: remainingDistance);
+
+      // Toggle whether this latitude's direction is reversed
+      if (bIsReversed) {
+        _reversedControls.remove(propertyIndex);
       } else {
-        _handleReverseTrackpad(sum, index);
+        _reversedControls.add(propertyIndex);
       }
-    }
-  }
-
-  /// Handle horizontal drag/pan across the trackpad widget providing Drag [details] containing positions on the x
-  /// and y axis.
-  void _onHorizontalDrag(DragUpdateDetails details) {
-    double dx = details.delta.dx * _deltaSettings.sensitivity.getValue();
-    for (var index = 0; index < _paths.length; index++) {
-      _modifyProperty(dx, index, _longitudeController, PropertyMinMaxBehaviour.loop);
-    }
-  }
-
-  /// Track all properties and update a property with [name] at the specified [index], modifying the property with
-  /// [value], using the related UnrealPropertyController [controller], optionally providing a
-  /// PropertyMinMaxBehaviour [behaviour] or [PropertyMinMaxBehaviour.clamp] will be used as default.
-  void _modifyProperty(double value, int index, UnrealPropertyController controller,
-      [PropertyMinMaxBehaviour behaviour = PropertyMinMaxBehaviour.clamp]) {
-    controller.modifyProperty(controller.modifyOperation, index, value: value, minMaxBehaviour: behaviour);
-  }
-
-  /// update tracked properties for both latitude and longitude controllers.
-  void _handleUpdateTrackedProperties() {
-    _latitudeController
-        .trackAllProperties(_getProperties('Latitude', _paths, modifierFn: _modifyPositionalPropertyNameBasedOnClass));
-    _longitudeController
-        .trackAllProperties(_getProperties('Longitude', _paths, modifierFn: _modifyPositionalPropertyNameBasedOnClass));
-  }
-
-  /// Handles all sort of drag/pan gestures on the trackpad, whether vertically or horizontally, providing Drag
-  /// update [details] which contains x [details.delta.dx] & y [details.delta.dy] positions.
-  void _onPanUpdate(DragUpdateDetails details) {
-    _longitudeController.beginTransaction();
-    _onVerticalDrag(details);
-    _onHorizontalDrag(details);
+    });
   }
 
-  /// Called when pan/drag gestures end's, providing Drag end [details].
+  /// Called when pan/drag gestures ends.
   void _onPanEnd(DragEndDetails details) {
-    _longitudeController.endTransaction();
-    _listBIsTrackpadReversed = _paths.map((e) => false).toList();
+    _onPanFinished();
   }
 
-  /// Called when pan/drag gesture is canceled.
+  /// Called when pan/drag gesture is cancelled.
   void _onPanCancel() {
-    _longitudeController.endTransaction();
-    _listBIsTrackpadReversed = _paths.map((e) => false).toList();
+    _onPanFinished();
   }
 
-  /// Called when pan/drag gesture starts.
-  void _onPanStart(DragStartDetails details) {
-    _listBIsTrackpadReversed = _paths.map((e) => false).toList();
+  /// Called when a pan gesture ends for any reason.
+  void _onPanFinished() {
+    _longitudeController.endTransaction();
+    _reversedControls.clear();
   }
 
-  /// Given an [actorPath] and a [propertyName], return a modified positional property name accounting for the actor's
-  /// type.
-  String _modifyPositionalPropertyNameBasedOnClass(String actorPath, String propertyName) {
-    final UnrealObject? actor = _actorManager.getActorAtPath(actorPath);
+  /// Get a list of UnrealProperties corresponding to the [propertyName] on all of the selected, controllable actors.
+  /// This will modify the property name to account for actor types where the parameters are in a nested struct.
+  List<UnrealProperty> _getPositionalProperties(String propertyName) {
+    return _getSelectedActorPaths().map((actorPath) {
+      final UnrealObject? actor = _actorManager.getActorAtPath(actorPath);
 
-    if (actor != null && actor.isAny(colorCorrectWindowClassNames)) {
-      // CCWs store their positional properties in a sub-structure.
-      return 'PositionalParams.' + propertyName;
-    }
+      final String pathToProperty;
 
-    return propertyName;
-  }
-}
+      if (actor?.isAny(colorCorrectWindowClassNames) == true) {
+        // CCWs store their positional properties in a sub-structure
+        pathToProperty = 'PositionalParams.' + propertyName;
+      } else {
+        pathToProperty = propertyName;
+      }
 
-/// A simple type definition used by the [_geProperties] function.
-typedef ModifierFn = String Function(String path, String name);
-
-/// Get a list of properties with the given [name] for all of the actors with path in [paths].
-/// If [modifierFn] is provided, it will be called for each [actorPath] and the [propertyName] of the property,
-/// and its return value will be used in place of [name].
-List<UnrealProperty> _getProperties(String name, List<String> paths, {ModifierFn? modifierFn, String? overrideName}) {
-  return paths
-      .map(
-        (actorPath) => UnrealProperty(
-          objectPath: actorPath,
-          propertyName: (modifierFn != null) ? modifierFn(actorPath, name) : name,
-          typeNameOverride: overrideName,
-        ),
-      )
-      .toList();
+      return UnrealProperty(objectPath: actorPath, propertyName: pathToProperty);
+    }).toList(growable: false);
+  }
 }
diff --git a/Engine/Extras/VirtualProduction/EpicStageApp/lib/widgets/screens/main/stage_app_main_screen.dart b/Engine/Extras/VirtualProduction/EpicStageApp/lib/widgets/screens/main/stage_app_main_screen.dart
--- a/Engine/Extras/VirtualProduction/EpicStageApp/lib/widgets/screens/main/stage_app_main_screen.dart
+++ b/Engine/Extras/VirtualProduction/EpicStageApp/lib/widgets/screens/main/stage_app_main_screen.dart
@@ -42,6 +42,9 @@ class _StageAppMainScreenState extends State<StageAppMainScreen>
   /// Whether to show the floating stage map preview.
   bool _bShowMapPreview = false;
 
+  /// Whether to show the floating trackpad.
+  bool _bShowTrackpad = false;
+
   @override
   void initState() {
     super.initState();
@@ -52,6 +55,7 @@ class _StageAppMainScreenState extends State<StageAppMainScreen>
       vsync: this,
     );
     _tabController.addListener(_onActiveTabChanged);
+    _onActiveTabChanged();
 
     final int startTab = Provider.of<MainScreenSettings>(context, listen: false).selectedTab.getValue();
 
@@ -109,7 +113,7 @@ class _StageAppMainScreenState extends State<StageAppMainScreen>
                 right: UnrealTheme.cardMargin,
                 bottom: UnrealTheme.cardMargin,
               ),
-              key: Key('Tab View'),
+              key: const Key('Tab View'),
               child: Container(
                 decoration: BoxDecoration(
                   color: Theme.of(context).colorScheme.background,
@@ -124,7 +128,7 @@ class _StageAppMainScreenState extends State<StageAppMainScreen>
               ),
             ),
             if (_bShowMapPreview) const FloatingMapPreview(key: Key('Map Preview')),
-            FloatingTrackpad()
+            if (_bShowTrackpad) const FloatingTrackpad(key: Key('Trackpad')),
           ],
         ),
       ),
@@ -183,6 +187,13 @@ class _StageAppMainScreenState extends State<StageAppMainScreen>
         _bShowMapPreview = bNewShowMapPreview;
       });
     }
+
+    final bool bNewShowTrackpad = MainScreenTabs.shouldShowTrackpad(_tabController.index);
+    if (bNewShowTrackpad != _bShowTrackpad) {
+      setState(() {
+        _bShowTrackpad = bNewShowTrackpad;
+      });
+    }
   }
 
   /// Called when a controllable actor has changed.
diff --git a/Engine/Extras/VirtualProduction/EpicStageApp/lib/widgets/screens/main/tabs/main_screen_tabs.dart b/Engine/Extras/VirtualProduction/EpicStageApp/lib/widgets/screens/main/tabs/main_screen_tabs.dart
--- a/Engine/Extras/VirtualProduction/EpicStageApp/lib/widgets/screens/main/tabs/main_screen_tabs.dart
+++ b/Engine/Extras/VirtualProduction/EpicStageApp/lib/widgets/screens/main/tabs/main_screen_tabs.dart
@@ -17,6 +17,7 @@ class MainScreenTabs {
       getTitle: StageMapTab.getTitle,
       iconPath: StageMapTab.iconPath,
       bShowMapPreview: false,
+      bShowTrackpad: true,
       createTabContents: (context) => const StageMapTab(),
     ),
     MainScreenTabConfig(
@@ -66,6 +67,11 @@ class MainScreenTabs {
   static bool shouldShowMapPreview(int index) {
     return _tabs[index].bShowMapPreview;
   }
+
+  /// Check whether the given tab should show the trackpad.
+  static bool shouldShowTrackpad(int index) {
+    return _tabs[index].bShowTrackpad;
+  }
 }
 
 /// Configuration for an entry in the main screen's tab bar.
@@ -77,7 +83,8 @@ class MainScreenTabConfig {
     this.bKeepAlive = false,
     this.bEnableOutlinerToggle = true,
     this.bShowMapPreview = true,
-  });
+    this.bShowTrackpad = false,
+  }) : assert(!(bShowMapPreview && bShowTrackpad), 'Map and trackpad can\'t be shown at the same time');
 
   /// Function that returns the title of the tab as shown to the user.
   final String Function(BuildContext) getTitle;
@@ -97,4 +104,7 @@ class MainScreenTabConfig {
 
   /// Whether the floating map preview should be visible when this tab is open.
   final bool bShowMapPreview;
+
+  /// Whether the floating trackpad should be visible when this tab is open.
+  final bool bShowTrackpad;
 }
