Version string bump for Flutter VCam App

diff --git a/Engine/Extras/VirtualProduction/FlutterLiveLinkVCAM/pubspec.yaml b/Engine/Extras/VirtualProduction/FlutterLiveLinkVCAM/pubspec.yaml
--- a/Engine/Extras/VirtualProduction/FlutterLiveLinkVCAM/pubspec.yaml
+++ b/Engine/Extras/VirtualProduction/FlutterLiveLinkVCAM/pubspec.yaml
@@ -2,7 +2,7 @@ name: live_link_vcam
 description: A rebuild of the Live Link VCAM app for Flutter, enabling cross-platform deployment.
 publish_to: 'none'
 
-version: 1.0.0+1
+version: 2.0.0+32867938
 
 environment:
   sdk: '>=3.1.0 <4.0.0'
@@ -86,4 +86,4 @@ flutter_launcher_icons:
   android: true
   ios: true
   remove_alpha_ios: true
-  image_path: "assets/app_icon/VCam_Icon_HD.png"
\ No newline at end of file
+  image_path: "assets/app_icon/VCam_Icon_HD.png"
