[LiveLinkVCAM] Update missing Android app icons

diff --git a/Engine/Extras/VirtualProduction/FlutterLiveLinkVCAM/android/app/src/main/res/drawable-v21/launch_background.xml b/Engine/Extras/VirtualProduction/FlutterLiveLinkVCAM/android/app/src/main/res/drawable-v21/launch_background.xml
--- a/Engine/Extras/VirtualProduction/FlutterLiveLinkVCAM/android/app/src/main/res/drawable-v21/launch_background.xml
+++ b/Engine/Extras/VirtualProduction/FlutterLiveLinkVCAM/android/app/src/main/res/drawable-v21/launch_background.xml
@@ -1,12 +1,12 @@
 <?xml version="1.0" encoding="utf-8"?>
 <!-- Modify this file to customize your launch splash screen -->
 <layer-list xmlns:android="http://schemas.android.com/apk/res/android">
-    <item android:drawable="?android:colorBackground" />
+    <item android:drawable="@android:color/black" />
 
     <!-- You can insert your own image assets here -->
-    <!-- <item>
+    <item>
         <bitmap
             android:gravity="center"
-            android:src="@mipmap/launch_image" />
-    </item> -->
+            android:src="@mipmap/ic_launcher" />
+    </item>
 </layer-list>
