[UnrealStage] Add final styling + fix case where new actors weren't controllable

diff --git a/Engine/Extras/VirtualProduction/EpicStageApp/lib/widgets/elements/floating_trackpad.dart b/Engine/Extras/VirtualProduction/EpicStageApp/lib/widgets/elements/floating_trackpad.dart
--- a/Engine/Extras/VirtualProduction/EpicStageApp/lib/widgets/elements/floating_trackpad.dart
+++ b/Engine/Extras/VirtualProduction/EpicStageApp/lib/widgets/elements/floating_trackpad.dart
@@ -1,7 +1,7 @@
 // Copyright Epic Games, Inc. All Rights Reserved.
 
-import 'dart:ui';
-
+import 'package:epic_common/theme.dart';
+import 'package:epic_common/widgets.dart';
 import 'package:flutter/material.dart';
 import 'package:flutter_gen/gen_l10n/app_localizations.dart';
 
@@ -12,7 +12,7 @@ class FloatingTrackpad extends StatelessWidget {
   const FloatingTrackpad({super.key});
 
   /// The height of the draggable area at the top of the trackpad.
-  static const double _dragBarHeight = 44;
+  static const double _dragBarHeight = 40;
 
   @override
   Widget build(BuildContext context) {
@@ -20,52 +20,35 @@ class FloatingTrackpad extends StatelessWidget {
 
     return FloatingWindow(
       settingsPrefix: 'floatingTrackpad',
-      icon: Icon(Icons.games_outlined),
-      builder: (context) => Stack(
+      icon: AssetIcon(path: 'assets/images/icons/trackpad.svg'),
+      builder: (context) => Column(
         children: [
-          Positioned(
-            top: _dragBarHeight - 1,
-            bottom: 0,
-            left: 0,
-            right: 0,
-            child: BackdropFilter(
-              filter: ImageFilter.blur(sigmaX: 3, sigmaY: 3),
-              child: Container(
-                color: Colors.black54,
-                child: Center(
-                  child: Icon(
-                    Icons.touch_app_rounded,
-                    size: 100,
-                    color: Colors.white12,
-                  ),
-                ),
-              ),
-            ),
-          ),
-          Positioned.fill(
-            child: LightCardTrackpad(),
-          ),
           Container(
             height: _dragBarHeight,
-            padding: const EdgeInsets.symmetric(horizontal: 4, vertical: 10),
+            padding: const EdgeInsets.symmetric(horizontal: 10, vertical: 10),
             alignment: Alignment.centerLeft,
             color: Theme.of(context).colorScheme.surfaceTint,
             child: Row(
               children: [
-                Icon(
-                  Icons.drag_indicator,
-                  color: Colors.white30,
+                AssetIcon(
+                  path: 'assets/images/icons/drag_handle.svg',
+                  color: UnrealColors.gray56,
+                  width: 18,
+                  height: 24,
                 ),
-                const SizedBox(width: 4),
+                const SizedBox(width: 12),
                 Expanded(
                   child: Text(
                     AppLocalizations.of(context)!.trackpadTitle.toUpperCase(),
-                    style: Theme.of(context).textTheme.titleLarge,
+                    style: Theme.of(context).textTheme.headlineSmall!.copyWith(color: UnrealColors.white),
                   ),
                 ),
               ],
             ),
           ),
+          Expanded(
+            child: LightCardTrackpad(),
+          ),
         ],
       ),
       draggableInsets: EdgeInsets.only(bottom: size.height - _dragBarHeight),
diff --git a/Engine/Extras/VirtualProduction/EpicStageApp/lib/widgets/elements/floating_window.dart b/Engine/Extras/VirtualProduction/EpicStageApp/lib/widgets/elements/floating_window.dart
--- a/Engine/Extras/VirtualProduction/EpicStageApp/lib/widgets/elements/floating_window.dart
+++ b/Engine/Extras/VirtualProduction/EpicStageApp/lib/widgets/elements/floating_window.dart
@@ -3,6 +3,7 @@
 import 'dart:math' as math;
 
 import 'package:epic_common/preferences.dart';
+import 'package:epic_common/theme.dart';
 import 'package:epic_common/widgets.dart';
 import 'package:flutter/material.dart';
 import 'package:flutter/physics.dart';
@@ -644,7 +645,7 @@ class _WindowBody extends StatelessWidget {
 
   @override
   Widget build(BuildContext context) {
-    final borderRadius = BorderRadius.circular(6);
+    final borderRadius = BorderRadius.circular(16);
 
     return SizedBox(
       width: size.width,
@@ -660,8 +661,8 @@ class _WindowBody extends StatelessWidget {
                 boxShadow: [
                   BoxShadow(
                     color: Color(0xff000000).withOpacity(bShadowVisible ? 0.4 : 0),
-                    spreadRadius: 2,
                     blurRadius: 5,
+                    blurStyle: BlurStyle.outer,
                   ),
                 ],
               ),
@@ -673,8 +674,8 @@ class _WindowBody extends StatelessWidget {
               decoration: BoxDecoration(
                 borderRadius: borderRadius,
                 border: Border.all(
-                  color: Color(0x50ffffff),
-                  width: 0.5,
+                  color: UnrealColors.gray56,
+                  width: 2,
                 ),
               ),
             ),
diff --git a/Engine/Extras/VirtualProduction/EpicStageApp/lib/widgets/elements/lightcard_trackpad.dart b/Engine/Extras/VirtualProduction/EpicStageApp/lib/widgets/elements/lightcard_trackpad.dart
--- a/Engine/Extras/VirtualProduction/EpicStageApp/lib/widgets/elements/lightcard_trackpad.dart
+++ b/Engine/Extras/VirtualProduction/EpicStageApp/lib/widgets/elements/lightcard_trackpad.dart
@@ -2,6 +2,8 @@
 
 import 'dart:async';
 
+import 'package:epic_common/theme.dart';
+import 'package:epic_common/widgets.dart';
 import 'package:flutter/material.dart';
 import 'package:provider/provider.dart';
 
@@ -53,6 +55,9 @@ class _LightCardTrackpadState extends State<LightCardTrackpad> {
   /// stream subscriptions for watching changes to selected actors.
   StreamSubscription? _selectedActorsSubscription;
 
+  /// True if the user is currently interacting with the trackpad.
+  bool _bIsPointerDown = false;
+
   @override
   void initState() {
     _actorManager = Provider.of<UnrealActorManager>(context, listen: false);
@@ -70,25 +75,43 @@ class _LightCardTrackpadState extends State<LightCardTrackpad> {
       }
     });
 
+    _actorManager.watchExistingSubscriptions(_onManagedActorsChanged);
+
     super.initState();
   }
 
   @override
   void dispose() {
     super.dispose();
+    _actorManager.stopWatchingExistingSubscriptions(_onManagedActorsChanged);
+    _selectedActorsSubscription?.cancel();
     _latitudeController.dispose();
     _longitudeController.dispose();
-    _selectedActorsSubscription?.cancel();
   }
 
   @override
   Widget build(BuildContext context) {
-    return GestureDetector(
-      behavior: HitTestBehavior.opaque,
-      onPanUpdate: _onPanUpdate,
-      onPanEnd: _onPanEnd,
-      onPanCancel: _onPanCancel,
-      onTap: () => print('trackpad'),
+    return Container(
+      color: UnrealColors.gray14.withOpacity(0.5),
+      child: GestureDetector(
+        behavior: HitTestBehavior.opaque,
+        onPanDown: _onPanDown,
+        onPanUpdate: _onPanUpdate,
+        onPanEnd: _onPanEnd,
+        onPanCancel: _onPanCancel,
+        child: Center(
+          child: AnimatedScale(
+            duration: const Duration(milliseconds: 200),
+            curve: Curves.easeOut,
+            scale: _bIsPointerDown ? 1.15 : 1,
+            child: AssetIcon(
+              path: 'assets/images/icons/trackpad.svg',
+              size: 96,
+              color: UnrealColors.white.withOpacity(0.2),
+            ),
+          ),
+        ),
+      ),
     );
   }
 
@@ -128,6 +151,13 @@ class _LightCardTrackpadState extends State<LightCardTrackpad> {
     _latitudeController.trackAllProperties(_getPositionalProperties('Latitude'));
   }
 
+  /// Called when the user's input suggests a pan gesture may be about to start.
+  void _onPanDown(DragDownDetails details) {
+    setState(() {
+      _bIsPointerDown = true;
+    });
+  }
+
   /// Called whenever a pan gesture's position updates based on new user input.
   void _onPanUpdate(DragUpdateDetails details) {
     // Transaction is global, so this will be shared with the latitude controller
@@ -215,6 +245,10 @@ class _LightCardTrackpadState extends State<LightCardTrackpad> {
 
   /// Called when a pan gesture ends for any reason.
   void _onPanFinished() {
+    setState(() {
+      _bIsPointerDown = false;
+    });
+
     _longitudeController.endTransaction();
     _reversedControls.clear();
   }
@@ -237,4 +271,27 @@ class _LightCardTrackpadState extends State<LightCardTrackpad> {
       return UnrealProperty(objectPath: actorPath, propertyName: pathToProperty);
     }).toList(growable: false);
   }
+
+  /// Called when any of the actors in the [UnrealActorManager] changes.
+  void _onManagedActorsChanged(ActorUpdateDetails details) {
+    if (!mounted) {
+      return;
+    }
+
+    final Set<String> newActorPaths = details.addedActors.map((actor) => actor.path).toSet();
+
+    // If any selected actors were added, update our tracked properties.
+    // This is necessary to catch new actors, which may be selected before they're registered with the actor manager.
+    bool bShouldUpdate = false;
+    for (final String actorPath in _selectedActorSettings.selectedActors.getValue()) {
+      if (newActorPaths.contains(actorPath)) {
+        bShouldUpdate = true;
+        break;
+      }
+    }
+
+    if (bShouldUpdate) {
+      _updateTrackedProperties();
+    }
+  }
 }
