[UnrealStage] Enable trackpad controls for UV light cards

diff --git a/Engine/Extras/VirtualProduction/EpicStageApp/lib/utilities/unreal_utilities.dart b/Engine/Extras/VirtualProduction/EpicStageApp/lib/utilities/unreal_utilities.dart
--- a/Engine/Extras/VirtualProduction/EpicStageApp/lib/utilities/unreal_utilities.dart
+++ b/Engine/Extras/VirtualProduction/EpicStageApp/lib/utilities/unreal_utilities.dart
@@ -137,7 +137,7 @@ bool canConvertUnrealType(String typeName) {
 /// Some conversion functions may also make make use of arbitrary data stored in [conversionMetadata].
 dynamic convertUnrealTypeJsonToDart(String typeName, dynamic data,
     {dynamic previousValue, Map<String, dynamic>? conversionMetadata}) {
-  if (data == null) {
+  if (data == null || (data is String && data.isEmpty)) {
     return null;
   }
 
diff --git a/Engine/Extras/VirtualProduction/EpicStageApp/lib/widgets/elements/lightcard_trackpad.dart b/Engine/Extras/VirtualProduction/EpicStageApp/lib/widgets/elements/lightcard_trackpad.dart
--- a/Engine/Extras/VirtualProduction/EpicStageApp/lib/widgets/elements/lightcard_trackpad.dart
+++ b/Engine/Extras/VirtualProduction/EpicStageApp/lib/widgets/elements/lightcard_trackpad.dart
@@ -6,6 +6,7 @@ import 'package:epic_common/theme.dart';
 import 'package:epic_common/widgets.dart';
 import 'package:flutter/material.dart';
 import 'package:provider/provider.dart';
+import 'package:vector_math/vector_math_64.dart' as vec;
 
 import '../../models/actor_data/light_card_actor_data.dart';
 import '../../models/property_modify_operations.dart';
@@ -31,7 +32,10 @@ class _LightCardTrackpadState extends State<LightCardTrackpad> {
   static const List<String> _controllableClasses = [lightCardClassName, ...colorCorrectWindowClassNames];
 
   /// Base change rate of latitude/longitude per logical pixel of change in the user's pointer position.
-  static const double baseDeltaMultiplier = 0.4;
+  static const double _positionalDeltaMultiplier = 0.4;
+
+  /// Base change rate of UV coordinates per logical pixel of change in the user's pointer position.
+  static const double _uvDeltaMultiplier = 0.004;
 
   /// Maximum value of latitude at which the "north" pole sits. "South" pole is assumed to be at the negation of this.
   static const _latitudeMax = 90;
@@ -49,6 +53,9 @@ class _LightCardTrackpadState extends State<LightCardTrackpad> {
   /// Property controller for longitude.
   late final UnrealPropertyController<double> _longitudeController;
 
+  /// Property controller for UV coordinates.
+  late final UnrealPropertyController<vec.Vector2> _uvController;
+
   /// Set of property indices for actors whose trackpad controls are reversed for the current drag operation.
   Set<int> _reversedControls = {};
 
@@ -66,6 +73,7 @@ class _LightCardTrackpadState extends State<LightCardTrackpad> {
 
     _latitudeController = UnrealPropertyController(context, bShouldInitTransaction: false);
     _longitudeController = UnrealPropertyController(context, bShouldInitTransaction: false);
+    _uvController = UnrealPropertyController(context, bShouldInitTransaction: false);
 
     _updateTrackedProperties();
 
@@ -87,6 +95,7 @@ class _LightCardTrackpadState extends State<LightCardTrackpad> {
     _selectedActorsSubscription?.cancel();
     _latitudeController.dispose();
     _longitudeController.dispose();
+    _uvController.dispose();
   }
 
   @override
@@ -116,7 +125,8 @@ class _LightCardTrackpadState extends State<LightCardTrackpad> {
   }
 
   /// Gets a list of selected actor paths, filtering out any actors that can't be controlled by the trackpad.
-  List<String> _getSelectedActorPaths() {
+  /// If [bGetUVActors] is true, return a list of only UV light cards. Otherwise, return a list of only non-UV actors.
+  List<String> _getSelectedActorPaths({required bool bGetUVActors}) {
     final List<String> validActorPaths = [];
 
     for (final String actorPath in _selectedActorSettings.selectedActors.getValue()) {
@@ -131,9 +141,9 @@ class _LightCardTrackpadState extends State<LightCardTrackpad> {
         continue;
       }
 
+      // Filter to UV/non-UV actors
       final LightCardActorData? lightCardActorData = actor.getPerClassData<LightCardActorData>();
-      if (lightCardActorData?.bIsUV == true) {
-        // This is an LC with UV coordinates, so it has no lat/long to control
+      if (lightCardActorData?.bIsUV != bGetUVActors) {
         continue;
       }
 
@@ -149,6 +159,11 @@ class _LightCardTrackpadState extends State<LightCardTrackpad> {
 
     _longitudeController.trackAllProperties(_getPositionalProperties('Longitude'));
     _latitudeController.trackAllProperties(_getPositionalProperties('Latitude'));
+    _uvController.trackAllProperties(
+      _getSelectedActorPaths(bGetUVActors: true)
+          .map((String actorPath) => UnrealProperty(objectPath: actorPath, propertyName: 'UVCoordinates'))
+          .toList(growable: false),
+    );
   }
 
   /// Called when the user's input suggests a pan gesture may be about to start.
@@ -163,14 +178,15 @@ class _LightCardTrackpadState extends State<LightCardTrackpad> {
     // Transaction is global, so this will be shared with the latitude controller
     _longitudeController.beginTransaction();
 
-    final Offset delta = details.delta * baseDeltaMultiplier * _deltaSettings.sensitivity.getValue();
+    final Offset baseDelta = details.delta * _deltaSettings.sensitivity.getValue();
+    final Offset positionalDelta = baseDelta * _positionalDeltaMultiplier;
 
     // Longitude is a simple loop operation at the min/max values, so we can update it directly
     _longitudeController.modifyProperties(
       _propertyOperation,
       values: List.generate(
         _longitudeController.properties.length,
-        (_) => delta.dx,
+        (_) => positionalDelta.dx,
         growable: false,
       ),
       minMaxBehaviour: PropertyMinMaxBehaviour.loop,
@@ -192,7 +208,7 @@ class _LightCardTrackpadState extends State<LightCardTrackpad> {
 
       final bool bIsReversed = _reversedControls.contains(propertyIndex);
 
-      double deltaY = delta.dy;
+      double deltaY = positionalDelta.dy;
 
       // Latitude moves the opposite direction of where the user would expect on a trackpad, so negate it by default
       if (!bIsReversed) {
@@ -231,6 +247,20 @@ class _LightCardTrackpadState extends State<LightCardTrackpad> {
         _reversedControls.add(propertyIndex);
       }
     });
+
+    // UV coordinates can be modified directly without any clamping/looping
+    final Offset uvDelta = baseDelta * _uvDeltaMultiplier;
+    final uvDeltaVector = vec.Vector2(uvDelta.dx, uvDelta.dy);
+
+    _uvController.modifyProperties(
+      _propertyOperation,
+      values: List.generate(
+        _uvController.properties.length,
+        (_) => uvDeltaVector,
+        growable: false,
+      ),
+      minMaxBehaviour: PropertyMinMaxBehaviour.ignore,
+    );
   }
 
   /// Called when pan/drag gestures ends.
@@ -253,10 +283,10 @@ class _LightCardTrackpadState extends State<LightCardTrackpad> {
     _reversedControls.clear();
   }
 
-  /// Get a list of UnrealProperties corresponding to the [propertyName] on all of the selected, controllable actors.
-  /// This will modify the property name to account for actor types where the parameters are in a nested struct.
+  /// Get a list of UnrealProperties corresponding to the [propertyName] on all of the selected, controllable non-UV
+  /// actors. This will modify the property name to account for actor types where the parameters are in a nested struct.
   List<UnrealProperty> _getPositionalProperties(String propertyName) {
-    return _getSelectedActorPaths().map((actorPath) {
+    return _getSelectedActorPaths(bGetUVActors: false).map((actorPath) {
       final UnrealObject? actor = _actorManager.getActorAtPath(actorPath);
 
       final String pathToProperty;
