[UnrealStage] Fix prompting for password when initial connection times out and clean up related edge cases

diff --git a/Engine/Extras/VirtualProduction/EpicStageApp/lib/models/engine_connection.dart b/Engine/Extras/VirtualProduction/EpicStageApp/lib/models/engine_connection.dart
--- a/Engine/Extras/VirtualProduction/EpicStageApp/lib/models/engine_connection.dart
+++ b/Engine/Extras/VirtualProduction/EpicStageApp/lib/models/engine_connection.dart
@@ -379,10 +379,7 @@ class EngineConnectionManager with WidgetsBindingObserver {
     // Set up the batch request and listen for its completion
     final batchRequest = UnrealHttpRequest(url: '/remote/batch', verb: 'PUT', body: {'Requests': requestParameters});
     final batchParameters = _generateHttpRequestParameters(batchRequest);
-    final message = {
-      'MessageName': 'http',
-      'Parameters': batchParameters,
-    };
+    final message = createUnrealWebSocketMessage('http', batchParameters);
 
     final Completer<UnrealHttpResponse> completer = _makeHttpResponseCompleterFromParameters(batchParameters);
     sendRawMessage(message);
@@ -557,10 +554,16 @@ class EngineConnectionManager with WidgetsBindingObserver {
     _saveLastConnection(connectionData);
 
     // Check passphrase
-    final bool bIsPassphraseValid = await _checkPassphrase().timeout(
-      Duration(seconds: 3),
-      onTimeout: () => false,
-    );
+    _log.info('Checking passphrase');
+    final bool bIsPassphraseValid;
+    try {
+      bIsPassphraseValid = await _checkPassphrase().timeout(Duration(seconds: 3));
+    } catch (e) {
+      _log.severe('Timed out waiting for passphrase check response');
+      disconnect();
+      return EngineConnectionResult.genericFailure;
+    }
+
     if (!bIsPassphraseValid) {
       // The HTTP handler will have already disconnected us, so we can just return the result
       return EngineConnectionResult.passphraseRejected;
@@ -768,6 +771,8 @@ class EngineConnectionManager with WidgetsBindingObserver {
       bIsHandshakeMessage: true,
     );
 
+    _log.info('Passphrase check returned with code ${response.code} / ${response.body}');
+
     return response.code != HttpResponseCode.unauthorized;
   }
 
diff --git a/Engine/Extras/VirtualProduction/EpicStageApp/lib/models/engine_passphrase_manager.dart b/Engine/Extras/VirtualProduction/EpicStageApp/lib/models/engine_passphrase_manager.dart
--- a/Engine/Extras/VirtualProduction/EpicStageApp/lib/models/engine_passphrase_manager.dart
+++ b/Engine/Extras/VirtualProduction/EpicStageApp/lib/models/engine_passphrase_manager.dart
@@ -14,14 +14,14 @@ class EnginePassphraseManager {
   }
 
   /// Set the [passphrase] associated with the given [connectionData], or delete it if no passphrase is given.
-  void setPassphrase(ConnectionData connectionData, String? passphrase) {
+  Future<void> setPassphrase(ConnectionData connectionData, String? passphrase) {
     final String key = _makeKeyForConnection(connectionData);
 
     if (passphrase == null) {
-      _secureStorage.delete(key: key);
-    } else {
-      _secureStorage.write(key: key, value: passphrase);
+      return _secureStorage.delete(key: key);
     }
+
+    return _secureStorage.write(key: key, value: passphrase);
   }
 
   /// Given [connectionData], create a unique key to refer to its passphrase in the secure storage database.
diff --git a/Engine/Extras/VirtualProduction/EpicStageApp/lib/widgets/screens/connect/views/passphrase_modal.dart b/Engine/Extras/VirtualProduction/EpicStageApp/lib/widgets/screens/connect/views/passphrase_modal.dart
--- a/Engine/Extras/VirtualProduction/EpicStageApp/lib/widgets/screens/connect/views/passphrase_modal.dart
+++ b/Engine/Extras/VirtualProduction/EpicStageApp/lib/widgets/screens/connect/views/passphrase_modal.dart
@@ -173,7 +173,7 @@ class _PassphraseModalDialogState extends State<PassphraseModalDialog> with Conn
     }
 
     final passphraseManager = Provider.of<EnginePassphraseManager>(context, listen: false);
-    passphraseManager.setPassphrase(_connectionData, _textController.text);
+    await passphraseManager.setPassphrase(_connectionData, _textController.text);
 
     final EngineConnectionResult result = await connect(
       _connectionData,
