Renamed .p4ignore files to .p4ignore.txt to match the convention we use.

The problem it caused is that some users have the expected P4IGNORE=.p4ignore.txt configuration setting while others don't. The issues is further exacerbated by the fact that last year Perforce handling of P4IGNORE changed last year. Previously, P4IGNORE was in effect only when explicitly set. However, now, when P4IGNORE is not present, it defaults to .p4ignore and .p4ignore.txt. What it all means is that users who had 'P4IGNORE=.p4ignore.txt' were not using the few incorrectly named .p4ignore, while users who didn't have P4IGNORE at all, were affected by them. It turns out that a few of the .p4ignore files had incorrect definition in the past, which led to Perforce issues that didn't seem reproducible on machines with the expected P4IGNORE configuration.

After renaming the files, all users, regardless of the environment, should be affected in the same way. The incorrect .p4ignore rules have been augmented in CL 35021865.

#rnx

diff --git a/Engine/Extras/FlutterCommon/.p4ignore.txt b/Engine/Extras/FlutterCommon/.p4ignore.txt
new file mode 100644
--- /dev/null
+++ b/Engine/Extras/FlutterCommon/.p4ignore.txt
@@ -0,0 +1,117 @@
+# Miscellaneous
+*.class
+*.log
+*.pyc
+*.swp
+.DS_Store
+.atom/
+.buildlog/
+.history
+.svn/
+
+# IntelliJ related
+*.iml
+*.ipr
+*.iws
+.idea/
+
+# The .vscode folder contains launch configuration and tasks you configure in
+# VS Code which you may wish to be included in version control, so this line
+# is commented out by default.
+#.vscode/
+
+# Flutter/Dart/Pub related
+**/doc/api/
+**/ios/Flutter/.last_build_id
+.dart_tool/
+.flutter-plugins
+.flutter-plugins-dependencies
+.packages
+.pub-cache/
+.pub/
+/build/
+
+# Web related
+lib/generated_plugin_registrant.dart
+
+# Symbolication related
+app.*.symbols
+
+# Obfuscation related
+app.*.map.json
+
+# Android Studio will place build artifacts here
+/android/app/debug
+/android/app/profile
+/android/app/release
+
+# Android
+gradle-wrapper.jar
+/android/.gradle
+/android/captures/
+/android/gradlew
+/android/gradlew.bat
+/android/local.properties
+/android/**/GeneratedPluginRegistrant.java
+
+# Remember to never publicly share your keystore.
+# See https://flutter.dev/docs/deployment/android#reference-the-keystore-from-the-app
+/android/key.properties
+/android/**/*.keystore
+/android/**/*.jks
+
+# iOS
+/ios/**/dgph
+/ios/*.mode1v3
+/ios/*.mode2v3
+/ios/*.moved-aside
+/ios/*.pbxuser
+/ios/*.perspectivev3
+/ios/**/*sync/
+/ios/.sconsign.dblite
+/ios/.tags*
+/ios/**/.vagrant/
+/ios/**/DerivedData/
+/ios/Icon?
+/ios/**/Pods/
+/ios/**/.symlinks/
+/ios/profile
+/ios/xcuserdata
+/ios/**/.generated/
+/ios/Flutter/App.framework
+/ios/Flutter/Flutter.framework
+/ios/Flutter/Flutter.podspec
+/ios/Flutter/Generated.xcconfig
+/ios/Flutter/ephemeral/
+/ios/Flutter/app.flx
+/ios/Flutter/app.zip
+/ios/Flutter/flutter_assets/
+/ios/Flutter/flutter_export_environment.sh
+/ios/ServiceDefinitions.json
+/ios/Runner/GeneratedPluginRegistrant.*
+
+# Exceptions to above rules.
+!/ios/default.mode1v3
+!/ios/default.mode2v3
+!/ios/default.pbxuser
+!/ios/default.perspectivev3
+
+# Windows
+/windows/flutter/ephemeral/
+/windows/flutter/generated_plugin*
+
+# Visual Studio user-specific files.
+/windows/*.suo
+/windows/*.user
+/windows/*.userosscache
+/windows/*.sln.docstates
+
+# Visual Studio build-related files.
+/windows/x64/
+/windows/x86/
+
+# Visual Studio cache files
+# files ending in .cache can be ignored
+/windows/*.[Cc]ache
+# but keep track of directories ending in .cache
+/windows/!*.[Cc]ache/
\ No newline at end of file
diff --git a/Engine/Extras/VirtualProduction/EpicStageApp/.p4ignore.txt b/Engine/Extras/VirtualProduction/EpicStageApp/.p4ignore.txt
new file mode 100644
--- /dev/null
+++ b/Engine/Extras/VirtualProduction/EpicStageApp/.p4ignore.txt
@@ -0,0 +1,117 @@
+# Miscellaneous
+*.class
+*.log
+*.pyc
+*.swp
+.DS_Store
+.atom/
+.buildlog/
+.history
+.svn/
+
+# IntelliJ related
+*.iml
+*.ipr
+*.iws
+.idea/
+
+# The .vscode folder contains launch configuration and tasks you configure in
+# VS Code which you may wish to be included in version control, so this line
+# is commented out by default.
+#.vscode/
+
+# Flutter/Dart/Pub related
+**/doc/api/
+**/ios/Flutter/.last_build_id
+.dart_tool/
+.flutter-plugins
+.flutter-plugins-dependencies
+.packages
+.pub-cache/
+.pub/
+/build/
+
+# Web related
+lib/generated_plugin_registrant.dart
+
+# Symbolication related
+app.*.symbols
+
+# Obfuscation related
+app.*.map.json
+
+# Android Studio will place build artifacts here
+/android/app/debug
+/android/app/profile
+/android/app/release
+
+# Android
+gradle-wrapper.jar
+/android/.gradle
+/android/captures/
+/android/gradlew
+/android/gradlew.bat
+/android/local.properties
+/android/**/GeneratedPluginRegistrant.java
+
+# Remember to never publicly share your keystore.
+# See https://flutter.dev/docs/deployment/android#reference-the-keystore-from-the-app
+/android/key.properties
+/android/**/*.keystore
+/android/**/*.jks
+
+# iOS
+/ios/**/dgph
+/ios/*.mode1v3
+/ios/*.mode2v3
+/ios/*.moved-aside
+/ios/*.pbxuser
+/ios/*.perspectivev3
+/ios/**/*sync/
+/ios/.sconsign.dblite
+/ios/.tags*
+/ios/**/.vagrant/
+/ios/**/DerivedData/
+/ios/Icon?
+/ios/**/Pods/
+/ios/**/.symlinks/
+/ios/profile
+/ios/xcuserdata
+/ios/**/.generated/
+/ios/Flutter/App.framework
+/ios/Flutter/Flutter.framework
+/ios/Flutter/Flutter.podspec
+/ios/Flutter/Generated.xcconfig
+/ios/Flutter/ephemeral/
+/ios/Flutter/app.flx
+/ios/Flutter/app.zip
+/ios/Flutter/flutter_assets/
+/ios/Flutter/flutter_export_environment.sh
+/ios/ServiceDefinitions.json
+/ios/Runner/GeneratedPluginRegistrant.*
+
+# Exceptions to above rules.
+!/ios/default.mode1v3
+!/ios/default.mode2v3
+!/ios/default.pbxuser
+!/ios/default.perspectivev3
+
+# Windows
+/windows/flutter/ephemeral/
+/windows/flutter/generated_plugin*
+
+# Visual Studio user-specific files.
+/windows/*.suo
+/windows/*.user
+/windows/*.userosscache
+/windows/*.sln.docstates
+
+# Visual Studio build-related files.
+/windows/x64/
+/windows/x86/
+
+# Visual Studio cache files
+# files ending in .cache can be ignored
+/windows/*.[Cc]ache
+# but keep track of directories ending in .cache
+/windows/!*.[Cc]ache/
\ No newline at end of file
diff --git a/Engine/Extras/VirtualProduction/FlutterLiveLinkVCAM/.p4ignore.txt b/Engine/Extras/VirtualProduction/FlutterLiveLinkVCAM/.p4ignore.txt
new file mode 100644
--- /dev/null
+++ b/Engine/Extras/VirtualProduction/FlutterLiveLinkVCAM/.p4ignore.txt
@@ -0,0 +1,123 @@
+# Miscellaneous
+*.class
+*.log
+*.pyc
+*.swp
+.DS_Store
+.atom/
+.buildlog/
+.history
+.svn/
+
+# IntelliJ related
+*.iml
+*.ipr
+*.iws
+.idea/
+
+# The .vscode folder contains launch configuration and tasks you configure in
+# VS Code which you may wish to be included in version control, so this line
+# is commented out by default.
+#.vscode/
+
+# Flutter/Dart/Pub related
+**/doc/api/
+**/IOS/Flutter/.last_build_id
+.dart_tool/
+.flutter-plugins
+.flutter-plugins-dependencies
+.packages
+.pub-cache/
+.pub/
+/build/
+
+# Web related
+lib/generated_plugin_registrant.dart
+
+# Symbolication related
+app.*.symbols
+
+# Obfuscation related
+app.*.map.json
+
+# Android Studio will place build artifacts here
+/android/app/debug
+/android/app/profile
+/android/app/release
+/android/key.properties
+/android/app/android-key.jks
+
+# Android
+gradle-wrapper.jar
+/android/.gradle
+/android/captures/
+/android/gradlew
+/android/gradlew.bat
+/android/local.properties
+/android/**/GeneratedPluginRegistrant.java
+
+# Android Tentacle SDK (copied in during Gradle build)
+/android/app/src/main/java/com/tentaclesync/
+
+# Remember to never publicly share your keystore.
+# See https://flutter.dev/docs/deployment/android#reference-the-keystore-from-the-app
+/android/key.properties
+/android/**/*.keystore
+/android/**/*.jks
+
+# iOS
+/IOS/**/dgph
+/IOS/*.mode1v3
+/IOS/*.mode2v3
+/IOS/*.moved-aside
+/IOS/*.pbxuser
+/IOS/*.perspectivev3
+/IOS/**/*sync/
+/IOS/.sconsign.dblite
+/IOS/.tags*
+/IOS/**/.vagrant/
+/IOS/**/DerivedData/
+/IOS/Icon?
+/IOS/**/Pods/
+/IOS/**/.symlinks/
+/IOS/profile
+/IOS/xcuserdata
+/IOS/**/.generated/
+/IOS/Flutter/App.framework
+/IOS/Flutter/Flutter.framework
+/IOS/Flutter/Flutter.podspec
+/IOS/Flutter/Generated.xcconfig
+/IOS/Flutter/ephemeral/
+/IOS/Flutter/app.flx
+/IOS/Flutter/app.zip
+/IOS/Flutter/flutter_assets/
+/IOS/Flutter/flutter_export_environment.sh
+/IOS/ServiceDefinitions.json
+/IOS/Runner/GeneratedPluginRegistrant.*
+/IOS/.symlinks
+
+# Exceptions to above rules.
+!/IOS/default.mode1v3
+!/IOS/default.mode2v3
+!/IOS/default.pbxuser
+!/IOS/default.perspectivev3
+
+# Windows
+/windows/flutter/ephemeral/
+/windows/flutter/generated_plugin*
+
+# Visual Studio user-specific files.
+/windows/*.suo
+/windows/*.user
+/windows/*.userosscache
+/windows/*.sln.docstates
+
+# Visual Studio build-related files.
+/windows/x64/
+/windows/x86/
+
+# Visual Studio cache files
+# files ending in .cache can be ignored
+/windows/*.[Cc]ache
+# but keep track of directories ending in .cache
+/windows/!*.[Cc]ache/
diff --git a/Engine/Extras/VirtualProduction/FlutterLiveLinkVCAM/IOS/.p4ignore.txt b/Engine/Extras/VirtualProduction/FlutterLiveLinkVCAM/IOS/.p4ignore.txt
new file mode 100644
--- /dev/null
+++ b/Engine/Extras/VirtualProduction/FlutterLiveLinkVCAM/IOS/.p4ignore.txt
@@ -0,0 +1,37 @@
+**/dgph
+*.mode1v3
+*.mode2v3
+*.moved-aside
+*.pbxuser
+*.perspectivev3
+**/*sync/
+.sconsign.dblite
+.tags*
+**/.vagrant/
+**/DerivedData/
+Icon?
+**/Pods/
+**/.symlinks/
+profile
+xcuserdata
+**/.generated/
+Flutter/App.framework
+Flutter/Flutter.framework
+Flutter/Flutter.podspec
+Flutter/Generated.xcconfig
+Flutter/ephemeral/
+Flutter/app.flx
+Flutter/app.zip
+Flutter/flutter_assets/
+Flutter/flutter_export_environment.sh
+ServiceDefinitions.json
+Runner/GeneratedPluginRegistrant.*
+
+# Exceptions to above rules.
+!default.mode1v3
+!default.mode2v3
+!default.pbxuser
+!default.perspectivev3
+
+.DS_Store
+Pods
diff --git a/Engine/Source/Runtime/SymsLib/.p4ignore.txt b/Engine/Source/Runtime/SymsLib/.p4ignore.txt
new file mode 100644
--- /dev/null
+++ b/Engine/Source/Runtime/SymsLib/.p4ignore.txt
@@ -0,0 +1,4 @@
+.vs/
+*.vxproj.user
+*.vcxproj.user
+obj/
