[UnrealStage] Add new blending mode property for lightcards

diff --git a/Engine/Extras/VirtualProduction/EpicStageApp/lib/models/api_version.dart b/Engine/Extras/VirtualProduction/EpicStageApp/lib/models/api_version.dart
--- a/Engine/Extras/VirtualProduction/EpicStageApp/lib/models/api_version.dart
+++ b/Engine/Extras/VirtualProduction/EpicStageApp/lib/models/api_version.dart
@@ -81,4 +81,7 @@ class EpicStageAppAPIVersion implements Comparable<EpicStageAppAPIVersion> {
 
   /// Whether compression of WebSocket traffic can be enabled in the engine.
   bool get bIsWebSocketCompressionAvailable => this >= const EpicStageAppAPIVersion(1, 9, 0);
+
+  /// Whether the engine supports the Blending Mode property for light cards.
+  bool get bIsLightCardBlendingModeAvailable => this >= const EpicStageAppAPIVersion(1, 10, 0);
 }
diff --git a/Engine/Extras/VirtualProduction/EpicStageApp/lib/widgets/screens/main/tabs/details_tab.dart b/Engine/Extras/VirtualProduction/EpicStageApp/lib/widgets/screens/main/tabs/details_tab.dart
--- a/Engine/Extras/VirtualProduction/EpicStageApp/lib/widgets/screens/main/tabs/details_tab.dart
+++ b/Engine/Extras/VirtualProduction/EpicStageApp/lib/widgets/screens/main/tabs/details_tab.dart
@@ -6,6 +6,8 @@ import 'dart:collection';
 import 'package:epic_common/preferences.dart';
 import 'package:epic_common/theme.dart';
 import 'package:epic_common/widgets.dart';
+import 'package:epic_stage_app/models/api_version.dart';
+import 'package:epic_stage_app/models/engine_connection.dart';
 import 'package:flutter/material.dart';
 import 'package:flutter_gen/gen_l10n/app_localizations.dart';
 import 'package:provider/provider.dart';
@@ -283,6 +285,9 @@ class _DetailsTabState extends State<DetailsTab> with GuardedRefreshState {
 
   /// Create property widgets for a lightcard.
   Widget _createLightCardPropertiesColumn() {
+    final engineConnection = Provider.of<EngineConnectionManager>(context, listen: false);
+    final EpicStageAppAPIVersion? apiVersion = engineConnection.apiVersion;
+
     return PreferenceBuilder(
       preference: _tabSettings.detailsPropertyDisplayType,
       builder: (context, DetailsPropertyDisplayType detailsPropertyDisplayType) {
@@ -295,6 +300,11 @@ class _DetailsTabState extends State<DetailsTab> with GuardedRefreshState {
                 overrideName: AppLocalizations.of(context)!.propertyLightCardMask,
                 unrealProperties: _getPropertiesOnValidActors('Mask'),
               ),
+              if (apiVersion?.bIsLightCardBlendingModeAvailable == true)
+                UnrealDropdownSelector(
+                  overrideName: AppLocalizations.of(context)!.propertyLightCardBlendingMode,
+                  unrealProperties: _getPropertiesOnValidActors('PerLightcardRenderMode'),
+                ),
               for (final Widget slider in _createCommonStageActorOrientationPropertyWidgets()) slider,
             ];
             break;
