[UnrealStage] Add support for integer-based sliders; add Sort Priority property; reduce min/max scale for UV light cards

diff --git a/Engine/Extras/VirtualProduction/EpicStageApp/lib/utilities/unreal_utilities.dart b/Engine/Extras/VirtualProduction/EpicStageApp/lib/utilities/unreal_utilities.dart
--- a/Engine/Extras/VirtualProduction/EpicStageApp/lib/utilities/unreal_utilities.dart
+++ b/Engine/Extras/VirtualProduction/EpicStageApp/lib/utilities/unreal_utilities.dart
@@ -49,6 +49,11 @@ Map<String, _UnrealTypeConversionData> _typeConverters = {
     },
     dartToJson: _noConversion.dartToJson,
   ),
+  'int32': _UnrealTypeConversionData(
+    jsonToDart: (dynamic jsonValue, dynamic previousValue, Map<String, dynamic> conversionMetadata) =>
+        jsonNumberToDouble(jsonValue, null),
+    dartToJson: (dynamic dartValue, dynamic conversionMetadata) => dartValue.round(),
+  ),
   'float': _UnrealTypeConversionData(
     jsonToDart: (dynamic jsonValue, dynamic previousValue, Map<String, dynamic> conversionMetadata) =>
         jsonNumberToDouble(jsonValue, null),
diff --git a/Engine/Extras/VirtualProduction/EpicStageApp/lib/widgets/elements/delta_slider.dart b/Engine/Extras/VirtualProduction/EpicStageApp/lib/widgets/elements/delta_slider.dart
--- a/Engine/Extras/VirtualProduction/EpicStageApp/lib/widgets/elements/delta_slider.dart
+++ b/Engine/Extras/VirtualProduction/EpicStageApp/lib/widgets/elements/delta_slider.dart
@@ -36,6 +36,7 @@ class UnrealDeltaSlider extends UnrealWidget {
     this.hardMin,
     this.hardMax,
     this.exponent = 1.0,
+    this.bIsInteger = false,
     this.bShowResetButton = true,
     this.fillColor,
     this.valueMessageOverride,
@@ -67,6 +68,9 @@ class UnrealDeltaSlider extends UnrealWidget {
   /// smaller at smaller slider values.
   final double exponent;
 
+  /// If true, treat the value as an integer by rounding it to the nearest whole value
+  final bool bIsInteger;
+
   /// If true, show the reset button.
   final bool bShowResetButton;
 
@@ -95,6 +99,7 @@ class _UnrealDeltaSliderState extends State<UnrealDeltaSlider> with UnrealWidget
         values: makeValueDataList(),
         bShowResetButton: bIsInResetMode && widget.bShowResetButton,
         bIsEnabled: !bIsInResetMode,
+        bIsInteger: widget.bIsInteger,
         onChanged: handleOnChangedByUser,
         onValuesSet: _onValuesSet,
         onReset: handleOnResetByUser,
@@ -131,6 +136,7 @@ class DrivenDeltaSlider extends StatefulWidget {
     this.bShowResetButton = true,
     this.bIgnoreSensitivity = false,
     this.bIsEnabled = true,
+    this.bIsInteger = false,
     this.defaultValue,
     this.onReset,
     this.onInteractionFinished,
@@ -168,6 +174,9 @@ class DrivenDeltaSlider extends StatefulWidget {
   /// If false, grey out the widget and don't accept inputs. Note that this doesn't affect the reset button.
   final bool bIsEnabled;
 
+  /// If true, treat the value as an integer by rounding it to the nearest whole value
+  final bool bIsInteger;
+
   /// Default value to reset to when the reset button is pressed.
   final double? defaultValue;
 
@@ -203,6 +212,9 @@ class _DrivenDeltaSliderState extends State<DrivenDeltaSlider> with DeltaWidgetS
   /// How long the min/max labels take to fade in/out.
   static const _minMaxFadeDuration = Duration(milliseconds: 100);
 
+  /// If in integer mode, this value is used to determine how close we are to crossing to the next integer
+  double _subIntegerPosition = 0;
+
   bool _bIsBeingDragged = false;
 
   @override
@@ -399,6 +411,10 @@ class _DrivenDeltaSliderState extends State<DrivenDeltaSlider> with DeltaWidgetS
 
   /// Format a value on the slider as a string for display to the user.
   String _getStringForValue(double value) {
+    if (widget.bIsInteger) {
+      return value.round().toString();
+    }
+
     return value.toStringAsFixed(2);
   }
 
@@ -415,7 +431,29 @@ class _DrivenDeltaSliderState extends State<DrivenDeltaSlider> with DeltaWidgetS
       return;
     }
 
-    final double delta = details.primaryDelta! * _valueSpan * deltaMultiplier / (context.size?.width ?? 1.0);
+    double delta = details.primaryDelta! * _valueSpan * deltaMultiplier / (context.size?.width ?? 1.0);
+
+    if (widget.bIsInteger) {
+      _subIntegerPosition += delta;
+
+      // Check if we've passed the integer threshold
+      bool bHasPassedThreshold = false;
+      if (_subIntegerPosition > 0.5) {
+        bHasPassedThreshold = true;
+        delta = (_subIntegerPosition - 0.5).ceilToDouble();
+      } else if (_subIntegerPosition < -0.5) {
+        bHasPassedThreshold = true;
+        delta = (_subIntegerPosition + 0.5).floorToDouble();
+      }
+
+      if (bHasPassedThreshold) {
+        _subIntegerPosition -= delta;
+      } else {
+        // The integer value hasn't changed, so don't report this change
+        return;
+      }
+    }
+
     final List<double> deltaValues = List.filled(widget.values.length, delta);
 
     if (widget.exponent != 1.0) {
@@ -443,6 +481,7 @@ class _DrivenDeltaSliderState extends State<DrivenDeltaSlider> with DeltaWidgetS
   void _onDragEnd() {
     widget.onInteractionFinished?.call();
     setState(() {
+      _subIntegerPosition = 0;
       _bIsBeingDragged = false;
     });
   }
@@ -462,18 +501,34 @@ class _DrivenDeltaSliderState extends State<DrivenDeltaSlider> with DeltaWidgetS
       return;
     }
 
-    final TextInputModalDialogResult<double>? result = await GenericModalDialogRoute.showDialog(
+    final String modalTitle = AppLocalizations.of(context)!.sliderKeyboardInputModalTitle(widget.label);
+    final double modalValue = widget.values[0]!.value;
+    final bool bShowResetButton = true;
+
+    final TextInputModalDialogResult? result = await GenericModalDialogRoute.showDialog(
       context: context,
-      builder: (context) => DoubleTextInputModalDialog(
-        title: AppLocalizations.of(context)!.sliderKeyboardInputModalTitle(widget.label),
-        initialValue: widget.values[0]!.value,
-        bShowResetButton: true,
-      ),
+      builder: (context) => widget.bIsInteger
+          ? IntegerTextInputModalDialog(
+              title: modalTitle,
+              initialValue: modalValue.round(),
+              bShowResetButton: bShowResetButton,
+            )
+          : DoubleTextInputModalDialog(
+              title: modalTitle,
+              initialValue: modalValue,
+              bShowResetButton: bShowResetButton,
+            ),
     );
 
     switch (result?.action) {
       case TextInputModalDialogAction.apply:
-        _setAllValues(result!.value!);
+        dynamic newValue = result!.value!;
+
+        if (widget.bIsInteger) {
+          newValue = newValue.toDouble();
+        }
+
+        _setAllValues(newValue);
         break;
 
       case TextInputModalDialogAction.reset:
diff --git a/Engine/Extras/VirtualProduction/EpicStageApp/lib/widgets/elements/property_divider.dart b/Engine/Extras/VirtualProduction/EpicStageApp/lib/widgets/elements/property_divider.dart
new file mode 100644
--- /dev/null
+++ b/Engine/Extras/VirtualProduction/EpicStageApp/lib/widgets/elements/property_divider.dart
@@ -0,0 +1,25 @@
+// Copyright Epic Games, Inc. All Rights Reserved.
+
+import 'package:flutter/material.dart';
+
+import 'delta_widget_base.dart';
+
+/// A divider used to split up groups of property widgets.
+class PropertyDivider extends StatelessWidget {
+  const PropertyDivider({super.key});
+
+  @override
+  Widget build(BuildContext context) {
+    return Padding(
+      padding: EdgeInsets.only(
+        left: DeltaWidgetConstants.widgetOuterXPadding,
+        right: DeltaWidgetConstants.widgetOuterXPadding,
+        bottom: 16,
+      ),
+      child: Container(
+        color: Theme.of(context).colorScheme.background,
+        height: 3,
+      ),
+    );
+  }
+}
diff --git a/Engine/Extras/VirtualProduction/EpicStageApp/lib/widgets/screens/main/tabs/details_tab.dart b/Engine/Extras/VirtualProduction/EpicStageApp/lib/widgets/screens/main/tabs/details_tab.dart
--- a/Engine/Extras/VirtualProduction/EpicStageApp/lib/widgets/screens/main/tabs/details_tab.dart
+++ b/Engine/Extras/VirtualProduction/EpicStageApp/lib/widgets/screens/main/tabs/details_tab.dart
@@ -6,14 +6,14 @@ import 'dart:collection';
 import 'package:epic_common/preferences.dart';
 import 'package:epic_common/theme.dart';
 import 'package:epic_common/widgets.dart';
-import 'package:epic_stage_app/models/api_version.dart';
-import 'package:epic_stage_app/models/engine_connection.dart';
 import 'package:flutter/material.dart';
 import 'package:flutter_gen/gen_l10n/app_localizations.dart';
 import 'package:provider/provider.dart';
 import 'package:streaming_shared_preferences/streaming_shared_preferences.dart';
 
 import '../../../../models/actor_data/light_card_actor_data.dart';
+import '../../../../models/api_version.dart';
+import '../../../../models/engine_connection.dart';
 import '../../../../models/property_modify_operations.dart';
 import '../../../../models/settings/details_tab_settings.dart';
 import '../../../../models/settings/main_screen_settings.dart';
@@ -25,6 +25,7 @@ import '../../../../utilities/guarded_refresh_state.dart';
 import '../../../elements/delta_slider.dart';
 import '../../../elements/dropdown_button.dart';
 import '../../../elements/dropdown_text.dart';
+import '../../../elements/property_divider.dart';
 import '../../../elements/reset_mode_button.dart';
 import '../../../elements/stepper.dart';
 import '../../../elements/unreal_property_builder.dart';
@@ -300,12 +301,24 @@ class _DetailsTabState extends State<DetailsTab> with GuardedRefreshState {
                 overrideName: AppLocalizations.of(context)!.propertyLightCardMask,
                 unrealProperties: _getPropertiesOnValidActors('Mask'),
               ),
+              for (final Widget slider in _createCommonStageActorOrientationPropertyWidgets()) slider,
+              const PropertyDivider(),
               if (apiVersion?.bIsLightCardBlendingModeAvailable == true)
                 UnrealDropdownSelector(
                   overrideName: AppLocalizations.of(context)!.propertyLightCardBlendingMode,
                   unrealProperties: _getPropertiesOnValidActors('PerLightcardRenderMode'),
                 ),
-              for (final Widget slider in _createCommonStageActorOrientationPropertyWidgets()) slider,
+              UnrealDeltaSlider(
+                overrideName: AppLocalizations.of(context)!.propertyLightCardSortPriority,
+                bIsInteger: true,
+                softMin: 0,
+                softMax: 5,
+                // This property lives on the light card component, which we assume to exist at this path
+                unrealProperties: _getPropertiesOnActors(
+                  'TranslucencySortPriority',
+                  _getValidActorPaths().map((path) => path + '.LightCard'),
+                ),
+              ),
             ];
             break;
 
@@ -506,6 +519,10 @@ class _DetailsTabState extends State<DetailsTab> with GuardedRefreshState {
 
     final List<String> scalableActorPaths = _getValidActorPathsOfClasses(scalableActorClasses);
 
+    final bool bHasNonUvLightCards = scalableActorPaths.any(
+      (actorPath) => _actorManager.getActorAtPath(actorPath)?.getPerClassData<LightCardActorData>()?.bIsUV != true,
+    );
+
     final List<Widget> scaleSliders = [];
     if (scalableActorPaths.isNotEmpty) {
       scaleSliders.addAll([
@@ -518,7 +535,7 @@ class _DetailsTabState extends State<DetailsTab> with GuardedRefreshState {
           ),
           overrideName: AppLocalizations.of(context)!.propertyLightCardScaleX,
           hardMin: 0,
-          softMax: 10,
+          softMax: bHasNonUvLightCards ? 10 : 1,
         ),
         UnrealDeltaSlider(
           key: Key('Scale Y'),
@@ -529,7 +546,7 @@ class _DetailsTabState extends State<DetailsTab> with GuardedRefreshState {
           ),
           overrideName: AppLocalizations.of(context)!.propertyLightCardScaleY,
           hardMin: 0,
-          softMax: 10,
+          softMax: bHasNonUvLightCards ? 10 : 1,
         ),
       ]);
     }
@@ -555,7 +572,7 @@ class _DetailsTabState extends State<DetailsTab> with GuardedRefreshState {
   /// and its return value will be used in place of [name].
   List<UnrealProperty> _getPropertiesOnActors(
     String name,
-    List<String> actorPaths, {
+    Iterable<String> actorPaths, {
     String Function(String actorPath, String propertyName)? modifierFunction,
     String? typeNameOverride,
   }) {
