[UnrealStage] Add toggle to hide unselected pins on map

diff --git a/Engine/Extras/VirtualProduction/EpicStageApp/lib/models/settings/stage_map_settings.dart b/Engine/Extras/VirtualProduction/EpicStageApp/lib/models/settings/stage_map_settings.dart
--- a/Engine/Extras/VirtualProduction/EpicStageApp/lib/models/settings/stage_map_settings.dart
+++ b/Engine/Extras/VirtualProduction/EpicStageApp/lib/models/settings/stage_map_settings.dart
@@ -26,6 +26,10 @@ class StageMapSettings {
           'stageMap.bIsInMultiSelectMode',
           defaultValue: false,
         ),
+        bShowOnlySelectedPins = preferences.persistent.getBool(
+          'stageMap.bShowOnlySelectedPins',
+          defaultValue: false,
+        ),
         perRootActorPreferences = preferences.persistent.getCustomValue(
           'stageMap.perRootActor',
           defaultValue: {},
@@ -72,6 +76,9 @@ class StageMapSettings {
   /// around.
   final Preference<bool> bIsInTransformMode;
 
+  /// If true, only show the currently selected pin on the map.
+  final Preference<bool> bShowOnlySelectedPins;
+
   /// The camera's current projection mode for stage actor locations.
   late final PerRootActorPreference<ProjectionMode> projectionMode;
 
diff --git a/Engine/Extras/VirtualProduction/EpicStageApp/lib/widgets/elements/lightcard_map.dart b/Engine/Extras/VirtualProduction/EpicStageApp/lib/widgets/elements/lightcard_map.dart
--- a/Engine/Extras/VirtualProduction/EpicStageApp/lib/widgets/elements/lightcard_map.dart
+++ b/Engine/Extras/VirtualProduction/EpicStageApp/lib/widgets/elements/lightcard_map.dart
@@ -238,7 +238,8 @@ class StageMapState extends State<StageMap> with PreviewRenderConsumer, GuardedR
               final Map<_StageMapPin, Offset> pinLocations = {};
               for (MapEntry<String, _MapActorData> actorPair in _mapActorDataByPath.entries) {
                 // Hide deselected pins if requested
-                if (widget.bShowOnlySelectedPins && !selectedActorSettings.isActorSelected(actorPair.key)) {
+                if ((stageMapSettings.bShowOnlySelectedPins.getValue() || widget.bShowOnlySelectedPins) &&
+                    !selectedActorSettings.isActorSelected(actorPair.key)) {
                   continue;
                 }
 
@@ -346,6 +347,7 @@ class StageMapState extends State<StageMap> with PreviewRenderConsumer, GuardedR
                               spacing: 16,
                               children: [
                                 const _StageMapControlModeToggle(),
+                                const _ShowOnlySelectedPinsToggle(),
 
                                 // View mode
                                 ConstrainedBox(
@@ -421,6 +423,7 @@ class StageMapState extends State<StageMap> with PreviewRenderConsumer, GuardedR
       selectedActorSettings.displayClusterRootPath.listen(_onDisplayClusterRootPathChanged),
       selectedActorSettings.selectedActors.listen((_) => _updateActorTransformRegistration()),
       stageMapSettings.bIsInTransformMode.listen(_onTransformModeChanged),
+      stageMapSettings.bShowOnlySelectedPins.listen((_) => guardedRefresh()),
     ]);
 
     _actorManager = Provider.of<UnrealActorManager>(context, listen: false);
@@ -1972,6 +1975,36 @@ class _StageMapControlModeToggleButton extends StatelessWidget {
   }
 }
 
+/// Unique widget that lets the user toggle the mode that only shows selected pins
+class _ShowOnlySelectedPinsToggle extends StatelessWidget {
+  const _ShowOnlySelectedPinsToggle({super.key});
+
+  @override
+  Widget build(BuildContext context) {
+    final stageMapSettings = Provider.of<StageMapSettings>(context, listen: false);
+
+    return PreferenceBuilder(
+      preference: stageMapSettings.bShowOnlySelectedPins,
+      builder: (_, final bool bShowOnlySelectedPins) => Container(
+        decoration: BoxDecoration(
+          color: Theme.of(context).colorScheme.surface,
+          borderRadius: BorderRadius.circular(16),
+        ),
+        clipBehavior: Clip.antiAlias,
+        child: EpicIconButton(
+          tooltipMessage: bShowOnlySelectedPins
+              ? AppLocalizations.of(context)!.stageMapShowOnlySelectedPinsOff
+              : AppLocalizations.of(context)!.stageMapShowOnlySelectedPinsOn,
+          iconPath: 'assets/images/icons/solo.svg',
+          buttonSize: const Size(48, 36),
+          bIsToggledOn: bShowOnlySelectedPins,
+          onPressed: () => stageMapSettings.bShowOnlySelectedPins.setValue(!bShowOnlySelectedPins),
+        ),
+      ),
+    );
+  }
+}
+
 /// Button and dropdown menu controlling the camera view in the stage map.
 class _StageMapViewDropdownButton extends StatelessWidget {
   const _StageMapViewDropdownButton();
