[UnrealStage] Add support for creating CCRs and modifying their transforms + types in the Details tab

diff --git a/Engine/Extras/VirtualProduction/EpicStageApp/lib/models/settings/recent_actor_settings.dart b/Engine/Extras/VirtualProduction/EpicStageApp/lib/models/settings/recent_actor_settings.dart
--- a/Engine/Extras/VirtualProduction/EpicStageApp/lib/models/settings/recent_actor_settings.dart
+++ b/Engine/Extras/VirtualProduction/EpicStageApp/lib/models/settings/recent_actor_settings.dart
@@ -11,6 +11,7 @@ enum RecentlyPlacedActorType {
   chromakeyCard,
   flag,
   colorCorrectWindow,
+  colorCorrectRegion,
 }
 
 /// Data about an actor the user recently placed.
diff --git a/Engine/Extras/VirtualProduction/EpicStageApp/lib/models/unreal_actor_creator.dart b/Engine/Extras/VirtualProduction/EpicStageApp/lib/models/unreal_actor_creator.dart
--- a/Engine/Extras/VirtualProduction/EpicStageApp/lib/models/unreal_actor_creator.dart
+++ b/Engine/Extras/VirtualProduction/EpicStageApp/lib/models/unreal_actor_creator.dart
@@ -134,7 +134,7 @@ class UnrealActorCreator {
       template: passedTemplate ?? defaultTemplate,
     );
 
-    createActor(className: lightCardClassName, settings: settings);
+    createStageActor(className: lightCardClassName, settings: settings);
 
     _recentActorSettings.addRecentlyPlacedActor(RecentlyPlacedActorData(
       name: passedTemplate?.name,
@@ -149,7 +149,7 @@ class UnrealActorCreator {
       settings = settings.copyWith(name: _localizations.newActorNameChromakeyCard);
     }
 
-    createActor(className: chromakeyCardClassName, settings: settings);
+    createStageActor(className: chromakeyCardClassName, settings: settings);
 
     _recentActorSettings.addRecentlyPlacedActor(RecentlyPlacedActorData(
       name: settings.template?.name,
@@ -191,7 +191,7 @@ class UnrealActorCreator {
       extraParams.addAll(settings.extraParams!);
     }
 
-    createActor(
+    createStageActor(
       className: lightCardClassName,
       settings: settings.copyWith(
         name: settings.name ?? _localizations.newActorNameFlag,
@@ -211,7 +211,7 @@ class UnrealActorCreator {
       settings = settings.copyWith(name: settings.name ?? _localizations.newActorNameColorCorrectWindow);
     }
 
-    createActor(
+    createStageActor(
       className: colorCorrectWindowClassNames.first,
       settings: settings.copyWith(
         mapPosition: settings.mapPosition ?? Offset(0.5, 0.5),
@@ -223,8 +223,28 @@ class UnrealActorCreator {
     ));
   }
 
-  /// Send a message to the engine requesting that a new actor be created.
-  void createActor({
+  /// Send a message to the engine requesting that a new CCR be created.
+  void createColorCorrectionRegion([ActorCreationSettings settings = ActorCreationSettings.defaults]) async {
+    if (settings.name == null) {
+      settings = settings.copyWith(name: settings.name ?? _localizations.newActorNameColorCorrectRegion);
+    }
+
+    createStageActor(
+      // colorCorrectRegionClassName (ColorCorrectRegion) is an abstract class, so we need to use the placeable class
+      // name here instead
+      className: '/Script/ColorCorrectRegions.ColorCorrectionRegion',
+      settings: settings.copyWith(
+        mapPosition: settings.mapPosition ?? Offset(0.5, 0.5),
+      ),
+    );
+
+    _recentActorSettings.addRecentlyPlacedActor(RecentlyPlacedActorData(
+      type: RecentlyPlacedActorType.colorCorrectRegion,
+    ));
+  }
+
+  /// Send a message to the engine requesting that a new stage actor be created.
+  void createStageActor({
     required String className,
     ActorCreationSettings settings = ActorCreationSettings.defaults,
   }) {
diff --git a/Engine/Extras/VirtualProduction/EpicStageApp/lib/models/unreal_property_controller.dart b/Engine/Extras/VirtualProduction/EpicStageApp/lib/models/unreal_property_controller.dart
--- a/Engine/Extras/VirtualProduction/EpicStageApp/lib/models/unreal_property_controller.dart
+++ b/Engine/Extras/VirtualProduction/EpicStageApp/lib/models/unreal_property_controller.dart
@@ -84,6 +84,9 @@ class UnrealPropertyController<PropertyType> with ChangeNotifier {
   /// Extra meta-data properties to pass to conversion functions operating on this widget's value.
   Map<String, dynamic>? get conversionMetadata => null;
 
+  /// Returns the number of times the tracked properties have changed.
+  int get trackedPropertyChangeCount => _latestTrackedPropertyVersion;
+
   /// Call this to indicate that the user has started interacting with the widget.
   /// This will be automatically called for you as soon as any properties change, but you can safely call it earlier.
   /// A [description] will be automatically generated if null is passed.
diff --git a/Engine/Extras/VirtualProduction/EpicStageApp/lib/widgets/elements/delta_slider.dart b/Engine/Extras/VirtualProduction/EpicStageApp/lib/widgets/elements/delta_slider.dart
--- a/Engine/Extras/VirtualProduction/EpicStageApp/lib/widgets/elements/delta_slider.dart
+++ b/Engine/Extras/VirtualProduction/EpicStageApp/lib/widgets/elements/delta_slider.dart
@@ -38,6 +38,7 @@ class UnrealDeltaSlider extends UnrealWidget {
     this.exponent = 1.0,
     this.bIsInteger = false,
     this.bShowResetButton = true,
+    this.bAreMinMaxRelative = false,
     this.fillColor,
     this.valueMessageOverride,
   });
@@ -74,6 +75,10 @@ class UnrealDeltaSlider extends UnrealWidget {
   /// If true, show the reset button.
   final bool bShowResetButton;
 
+  /// If true, the actual min/max is determined by adding the value of [min] and [max] (or 0) to the min/max values
+  /// in [values] whenever the selection changes (based on [selectionNumber]).
+  final bool bAreMinMaxRelative;
+
   /// The color to show in the filled portion of the bar (left of the value indicators).
   final Color? fillColor;
 
@@ -98,6 +103,7 @@ class _UnrealDeltaSliderState extends State<UnrealDeltaSlider> with UnrealWidget
       builder: (context, final bool bIsInResetMode) => DrivenDeltaSlider(
         values: makeValueDataList(),
         bShowResetButton: bIsInResetMode && widget.bShowResetButton,
+        bAreMinMaxRelative: widget.bAreMinMaxRelative,
         bIsEnabled: !bIsInResetMode,
         bIsInteger: widget.bIsInteger,
         onChanged: handleOnChangedByUser,
@@ -111,6 +117,7 @@ class _UnrealDeltaSliderState extends State<UnrealDeltaSlider> with UnrealWidget
         buildLabel: widget.buildLabel,
         fillColor: widget.fillColor,
         valueMessageOverride: widget.valueMessageOverride,
+        selectionNumber: trackedPropertyChangeCount,
       ),
     );
   }
@@ -134,6 +141,7 @@ class DrivenDeltaSlider extends StatefulWidget {
     this.max = 1.0,
     this.exponent = 1.0,
     this.bShowResetButton = true,
+    this.bAreMinMaxRelative = false,
     this.bIgnoreSensitivity = false,
     this.bIsEnabled = true,
     this.bIsInteger = false,
@@ -143,6 +151,7 @@ class DrivenDeltaSlider extends StatefulWidget {
     this.fillColor,
     this.valueMessageOverride,
     this.baseSensitivity,
+    this.selectionNumber,
   }) : super(key: key);
 
   /// Values to display on the slider.
@@ -168,6 +177,10 @@ class DrivenDeltaSlider extends StatefulWidget {
   /// If true, show the reset button.
   final bool bShowResetButton;
 
+  /// If true, the actual min/max is determined by adding the value of [min] and [max] (or -1/1) to the min/max values
+  /// in [values] whenever the selection changes (i.e. [selectionNumber] changes).
+  final bool bAreMinMaxRelative;
+
   /// If true, don't take sensitivity settings into account when calculating delta values.
   final bool bIgnoreSensitivity;
 
@@ -204,6 +217,9 @@ class DrivenDeltaSlider extends StatefulWidget {
   /// How sensitive the slider should be when the user's sensitivity is set to 1.0.
   final double? baseSensitivity;
 
+  /// When this changes, the slider's controlled selection is considered to have changed.
+  final int? selectionNumber;
+
   @override
   _DrivenDeltaSliderState createState() => _DrivenDeltaSliderState();
 }
@@ -212,11 +228,27 @@ class _DrivenDeltaSliderState extends State<DrivenDeltaSlider> with DeltaWidgetS
   /// How long the min/max labels take to fade in/out.
   static const _minMaxFadeDuration = Duration(milliseconds: 100);
 
+  /// The default minimum value.
+  static const double _defaultMin = 0;
+
+  /// The default maximum value.
+  static const double _defaultMax = 1;
+
   /// If in integer mode, this value is used to determine how close we are to crossing to the next integer
   double _subIntegerPosition = 0;
 
+  /// True if the slider is currently being dragged by the user.
   bool _bIsBeingDragged = false;
 
+  /// If true, we need to update [_minOverride] and [_maxOverride] when all values are non-null.
+  bool _bAreMinMaxOverridesDirty = true;
+
+  /// If set, the minimum will be no greater than this.
+  double? _minOverride;
+
+  /// If set, the minimum will be no less than this.
+  double? _maxOverride;
+
   @override
   bool get bIgnoreSensitivity => widget.bIgnoreSensitivity;
 
@@ -224,14 +256,18 @@ class _DrivenDeltaSliderState extends State<DrivenDeltaSlider> with DeltaWidgetS
   double get baseDeltaMultiplier => widget.baseSensitivity ?? super.baseDeltaMultiplier;
 
   /// Get the slider's minimum value.
-  /// If the widget has a set minimum, return it. Otherwise, default to a minimum of 0, but extend the minimum to match
-  /// the lowest controlled value.
+  /// If the widget has a set minimum, return it. Otherwise, default to a minimum of [_defaultMin], but extend the
+  /// minimum to match the lowest controlled value.
   double get _min {
+    if (_minOverride != null) {
+      return _minOverride!;
+    }
+
     if (widget.min != null) {
       return widget.min!;
     }
 
-    double minValue = 0.0;
+    double minValue = _defaultMin;
     for (final DeltaSliderValueData? valueData in widget.values) {
       if (valueData != null) {
         minValue = math.min(valueData.value, minValue);
@@ -242,14 +278,18 @@ class _DrivenDeltaSliderState extends State<DrivenDeltaSlider> with DeltaWidgetS
   }
 
   /// Get the slider's maximum value.
-  /// If the widget has a set maximum, return it. Otherwise, default to a maximum of 1, but extend the maximum to match
-  /// the highest controlled value.
+  /// If the widget has a set maximum, return it. Otherwise, default to a maximum of [_defaultMax], but extend the
+  /// maximum to match the highest controlled value.
   double get _max {
+    if (_maxOverride != null) {
+      return _maxOverride!;
+    }
+
     if (widget.max != null) {
       return widget.max!;
     }
 
-    double maxValue = 1.0;
+    double maxValue = _defaultMax;
     for (final DeltaSliderValueData? valueData in widget.values) {
       if (valueData != null) {
         maxValue = math.max(valueData.value, maxValue);
@@ -292,6 +332,21 @@ class _DrivenDeltaSliderState extends State<DrivenDeltaSlider> with DeltaWidgetS
     return _getStringForValue(firstValue);
   }
 
+  @override
+  void didUpdateWidget(DrivenDeltaSlider oldWidget) {
+    super.didUpdateWidget(oldWidget);
+
+    if (oldWidget.selectionNumber != widget.selectionNumber) {
+      // Selection has changed, so flag that we may need to update the min/max
+      _bAreMinMaxOverridesDirty = true;
+    }
+
+    // This is in a separate check because we may need to run this multiple times before all values are non-null
+    if (widget.bAreMinMaxRelative && _bAreMinMaxOverridesDirty) {
+      _updateMinMaxOverrides();
+    }
+  }
+
   @override
   Widget build(BuildContext context) {
     final String? valueString = widget.valueMessageOverride ?? _singleSharedValueString;
@@ -569,6 +624,21 @@ class _DrivenDeltaSliderState extends State<DrivenDeltaSlider> with DeltaWidgetS
       widget.onChanged(deltas);
     }
   }
+
+  /// If any values are outside the default min/max, adjust the min/max based on the current values (if ready).
+  void _updateMinMaxOverrides() {
+    if (widget.values.isEmpty || widget.values.any((value) => value == null)) {
+      // Don't adjust until all values are known
+      return;
+    }
+
+    final Iterable<double> values = widget.values.map((data) => data!.value);
+
+    _minOverride = values.reduce(math.min) + (widget.min ?? -1);
+    _maxOverride = values.reduce(math.max) + (widget.max ?? 1);
+
+    _bAreMinMaxOverridesDirty = false;
+  }
 }
 
 /// The interactive slider within the slider widget.
diff --git a/Engine/Extras/VirtualProduction/EpicStageApp/lib/widgets/elements/dropdown_button.dart b/Engine/Extras/VirtualProduction/EpicStageApp/lib/widgets/elements/dropdown_button.dart
--- a/Engine/Extras/VirtualProduction/EpicStageApp/lib/widgets/elements/dropdown_button.dart
+++ b/Engine/Extras/VirtualProduction/EpicStageApp/lib/widgets/elements/dropdown_button.dart
@@ -93,49 +93,17 @@ abstract class _UnrealDropdownSelectorOuterState<WidgetType extends UnrealWidget
   @override
   Widget build(BuildContext context) {
     final SingleSharedValue<PropertyType> sharedValue = getSingleSharedValue();
-    final TextStyle textStyle = Theme.of(context).textTheme.labelMedium!;
 
-    return Padding(
-      padding: EdgeInsets.only(bottom: 16),
-      child: TransientPreferenceBuilder(
-        preference: Provider.of<DeltaWidgetSettings>(context, listen: false).bIsInResetMode,
-        builder: (context, final bool bIsInResetMode) => Row(
-          children: [
-            Expanded(
-              child: Padding(
-                padding: EdgeInsets.symmetric(horizontal: DeltaWidgetConstants.widgetOuterXPadding),
-                child: Column(
-                  mainAxisAlignment: MainAxisAlignment.start,
-                  crossAxisAlignment: CrossAxisAlignment.start,
-                  children: [
-                    Text(
-                      propertyLabel,
-                      style: textStyle.copyWith(
-                        color: textStyle.color!.withOpacity(bIsInResetMode ? 0.4 : 1.0),
-                      ),
-                    ),
-                    Padding(
-                      padding: const EdgeInsets.only(
-                        left: DeltaWidgetConstants.widgetInnerXPadding,
-                        right: DeltaWidgetConstants.widgetInnerXPadding,
-                        top: 4,
-                      ),
-                      child: buildDropdownSelector(
-                        bIsInResetMode: bIsInResetMode,
-                        sharedValue: sharedValue,
-                        hint: sharedValue.bHasMultipleValues ? AppLocalizations.of(context)!.mismatchedValuesLabel : '',
-                      ),
-                    ),
-                  ],
-                ),
-              ),
-            ),
-            if (bIsInResetMode)
-              Padding(
-                padding: const EdgeInsets.only(right: DeltaWidgetConstants.widgetOuterXPadding),
-                child: ResetValueButton(onPressed: handleOnResetByUser),
-              ),
-          ],
+    return TransientPreferenceBuilder(
+      preference: Provider.of<DeltaWidgetSettings>(context, listen: false).bIsInResetMode,
+      builder: (context, final bool bIsInResetMode) => DropdownSelectorPropertyWrapper(
+        propertyLabel: propertyLabel,
+        bIsInResetMode: bIsInResetMode,
+        onReset: handleOnResetByUser,
+        child: buildDropdownSelector(
+          bIsInResetMode: bIsInResetMode,
+          sharedValue: sharedValue,
+          hint: sharedValue.bHasMultipleValues ? AppLocalizations.of(context)!.mismatchedValuesLabel : '',
         ),
       ),
     );
@@ -162,3 +130,70 @@ abstract class _UnrealDropdownSelectorOuterState<WidgetType extends UnrealWidget
     endTransaction();
   }
 }
+
+/// Wrapper for a [DropdownSelector] that styles it as an Unreal property.
+/// This only applies styling; see UnrealDropdownSelector for an implementation that actually talks to the engine.
+class DropdownSelectorPropertyWrapper extends StatelessWidget {
+  const DropdownSelectorPropertyWrapper({
+    super.key,
+    required this.child,
+    required this.propertyLabel,
+    this.bIsInResetMode = false,
+    this.onReset,
+  });
+
+  /// The DropdownSelector to display inside the wrapper.
+  final Widget child;
+
+  /// The label of the property to show to users.
+  final String propertyLabel;
+
+  /// If true, show a reset button next to the widget.
+  final bool bIsInResetMode;
+
+  /// Callback function for when the reset button is pressed.
+  final Function()? onReset;
+
+  @override
+  Widget build(BuildContext context) {
+    final TextStyle textStyle = Theme.of(context).textTheme.labelMedium!;
+
+    return Padding(
+      padding: const EdgeInsets.only(bottom: 16),
+      child: Row(
+        children: [
+          Expanded(
+            child: Padding(
+              padding: EdgeInsets.symmetric(horizontal: DeltaWidgetConstants.widgetOuterXPadding),
+              child: Column(
+                mainAxisAlignment: MainAxisAlignment.start,
+                crossAxisAlignment: CrossAxisAlignment.start,
+                children: [
+                  Text(
+                    propertyLabel,
+                    style: textStyle.copyWith(
+                      color: textStyle.color!.withOpacity(bIsInResetMode ? 0.4 : 1.0),
+                    ),
+                  ),
+                  Padding(
+                    padding: const EdgeInsets.only(
+                      left: DeltaWidgetConstants.widgetInnerXPadding,
+                      right: DeltaWidgetConstants.widgetInnerXPadding,
+                      top: 4,
+                    ),
+                    child: child,
+                  ),
+                ],
+              ),
+            ),
+          ),
+          if (bIsInResetMode)
+            Padding(
+              padding: const EdgeInsets.only(right: DeltaWidgetConstants.widgetOuterXPadding),
+              child: ResetValueButton(onPressed: onReset),
+            ),
+        ],
+      ),
+    );
+  }
+}
diff --git a/Engine/Extras/VirtualProduction/EpicStageApp/lib/widgets/elements/place_actor_menu.dart b/Engine/Extras/VirtualProduction/EpicStageApp/lib/widgets/elements/place_actor_menu.dart
--- a/Engine/Extras/VirtualProduction/EpicStageApp/lib/widgets/elements/place_actor_menu.dart
+++ b/Engine/Extras/VirtualProduction/EpicStageApp/lib/widgets/elements/place_actor_menu.dart
@@ -10,7 +10,6 @@ import 'package:streaming_shared_preferences/streaming_shared_preferences.dart';
 import '../../models/settings/recent_actor_settings.dart';
 import '../../models/unreal_actor_creator.dart';
 import '../../models/unreal_types.dart';
-import 'dropdown_button.dart';
 import 'template_menu.dart';
 
 /// Button that opens the actor placement menu.
@@ -138,6 +137,14 @@ class PlaceActorDropDownMenu extends StatelessWidget {
             _closeMenu(context);
           },
         ),
+        ListMenuSimpleItem(
+          title: AppLocalizations.of(context)!.actorNameColorCorrectRegion,
+          iconPath: 'packages/epic_common/assets/icons/color_correct_region.svg',
+          onTap: () {
+            actorCreator.createColorCorrectionRegion(actorCreationSettings);
+            _closeMenu(context);
+          },
+        ),
         ListMenuSimpleItem(
           title: AppLocalizations.of(context)!.placeActorMenuAllTemplates,
           iconPath: 'packages/epic_common/assets/icons/template.svg',
@@ -204,6 +211,9 @@ class _RecentlyPlacedMenuItem extends StatelessWidget {
       case RecentlyPlacedActorType.colorCorrectWindow:
         return 'packages/epic_common/assets/icons/color_correct_window.svg';
 
+      case RecentlyPlacedActorType.colorCorrectRegion:
+        return 'packages/epic_common/assets/icons/color_correct_region.svg';
+
       default:
         return null;
     }
@@ -236,6 +246,9 @@ class _RecentlyPlacedMenuItem extends StatelessWidget {
 
       case RecentlyPlacedActorType.colorCorrectWindow:
         return AppLocalizations.of(context)!.actorNameColorCorrectWindow;
+
+      case RecentlyPlacedActorType.colorCorrectRegion:
+        return AppLocalizations.of(context)!.actorNameColorCorrectRegion;
     }
   }
 
@@ -270,6 +283,10 @@ class _RecentlyPlacedMenuItem extends StatelessWidget {
           actorCreator.createColorCorrectionWindow(actorSettings);
           break;
 
+        case RecentlyPlacedActorType.colorCorrectRegion:
+          actorCreator.createColorCorrectionRegion(actorSettings);
+          break;
+
         default:
           throw 'Can\'t create actor of type ${actorData.type}';
       }
diff --git a/Engine/Extras/VirtualProduction/EpicStageApp/lib/widgets/elements/property_divider.dart b/Engine/Extras/VirtualProduction/EpicStageApp/lib/widgets/elements/property_divider.dart
--- a/Engine/Extras/VirtualProduction/EpicStageApp/lib/widgets/elements/property_divider.dart
+++ b/Engine/Extras/VirtualProduction/EpicStageApp/lib/widgets/elements/property_divider.dart
@@ -6,7 +6,10 @@ import 'delta_widget_base.dart';
 
 /// A divider used to split up groups of property widgets.
 class PropertyDivider extends StatelessWidget {
-  const PropertyDivider({super.key});
+  const PropertyDivider({super.key, this.padding = 16});
+
+  /// Padding under the divider.
+  final double padding;
 
   @override
   Widget build(BuildContext context) {
@@ -14,7 +17,7 @@ class PropertyDivider extends StatelessWidget {
       padding: EdgeInsets.only(
         left: DeltaWidgetConstants.widgetOuterXPadding,
         right: DeltaWidgetConstants.widgetOuterXPadding,
-        bottom: 16,
+        bottom: padding,
       ),
       child: Container(
         color: Theme.of(context).colorScheme.background,
diff --git a/Engine/Extras/VirtualProduction/EpicStageApp/lib/widgets/elements/unreal_widget_base.dart b/Engine/Extras/VirtualProduction/EpicStageApp/lib/widgets/elements/unreal_widget_base.dart
--- a/Engine/Extras/VirtualProduction/EpicStageApp/lib/widgets/elements/unreal_widget_base.dart
+++ b/Engine/Extras/VirtualProduction/EpicStageApp/lib/widgets/elements/unreal_widget_base.dart
@@ -129,6 +129,9 @@ mixin UnrealWidgetStateMixin<WidgetType extends UnrealWidget, PropertyType> on S
   /// The property's possible enum values as received from the engine.
   List<String> get propertyEnumValues => _singlePropertyController.propertyEnumValues;
 
+  /// Returns the number of times the tracked properties have changed.
+  int get trackedPropertyChangeCount => _singlePropertyController.trackedPropertyChangeCount;
+
   @override
   void initState() {
     _singlePropertyController = UnrealPropertyController(context);
diff --git a/Engine/Extras/VirtualProduction/EpicStageApp/lib/widgets/screens/main/tabs/details_tab.dart b/Engine/Extras/VirtualProduction/EpicStageApp/lib/widgets/screens/main/tabs/details_tab.dart
--- a/Engine/Extras/VirtualProduction/EpicStageApp/lib/widgets/screens/main/tabs/details_tab.dart
+++ b/Engine/Extras/VirtualProduction/EpicStageApp/lib/widgets/screens/main/tabs/details_tab.dart
@@ -8,6 +8,7 @@ import 'package:epic_common/theme.dart';
 import 'package:epic_common/widgets.dart';
 import 'package:flutter/material.dart';
 import 'package:flutter_gen/gen_l10n/app_localizations.dart';
+import 'package:logging/logging.dart';
 import 'package:provider/provider.dart';
 import 'package:streaming_shared_preferences/streaming_shared_preferences.dart';
 
@@ -19,6 +20,7 @@ import '../../../../models/settings/details_tab_settings.dart';
 import '../../../../models/settings/main_screen_settings.dart';
 import '../../../../models/settings/selected_actor_settings.dart';
 import '../../../../models/unreal_actor_manager.dart';
+import '../../../../models/unreal_transaction_manager.dart';
 import '../../../../models/unreal_types.dart';
 import '../../../../utilities/constants.dart';
 import '../../../../utilities/guarded_refresh_state.dart';
@@ -32,6 +34,8 @@ import '../../../elements/unreal_property_builder.dart';
 import '../sidebar/outliner_panel.dart';
 import 'base_color_tab.dart';
 
+final _log = Logger('DetailsTab');
+
 /// Which type of actor is being shown in the details tab.
 enum _DetailsActorType {
   lightCard,
@@ -231,23 +235,12 @@ class _DetailsTabState extends State<DetailsTab> with GuardedRefreshState {
       );
     }
 
-    // Buttons to toggle between property display types
-    Widget propertyDisplayTypeToggleButtons = PreferenceBuilder(
-      preference: _tabSettings.detailsPropertyDisplayType,
-      builder: (context, DetailsPropertyDisplayType detailsPropertyDisplayType) => SelectorBar(
-        value: detailsPropertyDisplayType,
-        onSelected: (DetailsPropertyDisplayType value) => _tabSettings.detailsPropertyDisplayType.setValue(value),
-        valueNames: {
-          DetailsPropertyDisplayType.orientation: AppLocalizations.of(context)!.colorTabOrientationPropertiesLabel,
-          DetailsPropertyDisplayType.appearance: AppLocalizations.of(context)!.colorTabAppearancePropertiesLabel,
-        },
-      ),
-    );
-
     late final BaseColorTabMode mode;
     late final List<UnrealProperty> colorProperties;
     late final Widget otherPropertiesColumn;
 
+    String orientationLabel = AppLocalizations.of(context)!.colorTabOrientationPropertiesLabel;
+
     switch (actorType) {
       case _DetailsActorType.lightCard:
         mode = BaseColorTabMode.color;
@@ -267,9 +260,8 @@ class _DetailsTabState extends State<DetailsTab> with GuardedRefreshState {
         colorProperties = _getPropertiesOnValidActors('ColorGradingSettings');
         otherPropertiesColumn = _createColorCorrectRegionPropertiesColumn();
 
-        // Can't control CCR orientation properties
-        propertyDisplayTypeToggleButtons =
-            FakeSelectorBar(AppLocalizations.of(context)!.colorTabAppearancePropertiesLabel);
+        // CCRs show Transform section instead
+        orientationLabel = AppLocalizations.of(context)!.colorTabTransformPropertiesLabel;
         break;
 
       default:
@@ -279,7 +271,19 @@ class _DetailsTabState extends State<DetailsTab> with GuardedRefreshState {
     return BaseColorTab(
       colorProperties: colorProperties,
       mode: mode,
-      rightSideHeader: Center(child: propertyDisplayTypeToggleButtons),
+      rightSideHeader: Center(
+        child: PreferenceBuilder(
+          preference: _tabSettings.detailsPropertyDisplayType,
+          builder: (context, DetailsPropertyDisplayType detailsPropertyDisplayType) => SelectorBar(
+            value: detailsPropertyDisplayType,
+            onSelected: (DetailsPropertyDisplayType value) => _tabSettings.detailsPropertyDisplayType.setValue(value),
+            valueNames: {
+              DetailsPropertyDisplayType.orientation: orientationLabel,
+              DetailsPropertyDisplayType.appearance: AppLocalizations.of(context)!.colorTabAppearancePropertiesLabel,
+            },
+          ),
+        ),
+      ),
       rightSideContents: otherPropertiesColumn,
     );
   }
@@ -313,10 +317,9 @@ class _DetailsTabState extends State<DetailsTab> with GuardedRefreshState {
                 bIsInteger: true,
                 softMin: 0,
                 softMax: 5,
-                // This property lives on the light card component, which we assume to exist at this path
-                unrealProperties: _getPropertiesOnActors(
+                unrealProperties: _getPropertiesOnValidActors(
                   'TranslucencySortPriority',
-                  _getValidActorPaths().map((path) => path + '.LightCard'),
+                  subObjectName: 'LightCard',
                 ),
               ),
             ];
@@ -364,35 +367,144 @@ class _DetailsTabState extends State<DetailsTab> with GuardedRefreshState {
 
   /// Create property widgets for a CCR.
   Widget _createColorCorrectRegionPropertiesColumn() {
-    final temperatureTypeProperties = _getPropertiesOnValidActors('TemperatureType');
-
-    return Column(children: [
-      UnrealMultiPropertyBuilder<String>(
-        properties: temperatureTypeProperties,
-        fallbackValue: AppLocalizations.of(context)!.mismatchedValuesLabel,
-        builder: (_, String? sharedValue, __) => UnrealDeltaSlider(
-          overrideName: sharedValue,
-          unrealProperties: _getPropertiesOnValidActors('Temperature'),
-          buildLabel: (name) => UnrealDropdownText(unrealProperties: temperatureTypeProperties),
-        ),
-      ),
-      UnrealDeltaSlider(
-        unrealProperties: _getPropertiesOnValidActors('Tint'),
-      ),
-      UnrealDeltaSlider(
-        unrealProperties: _getPropertiesOnValidActors('Intensity'),
-      ),
-      UnrealDeltaSlider(
-        unrealProperties: _getPropertiesOnValidActors('Inner'),
-      ),
-      UnrealDeltaSlider(
-        unrealProperties: _getPropertiesOnValidActors('Outer'),
-      ),
-      UnrealDeltaSlider(
-        unrealProperties: _getPropertiesOnValidActors('Falloff'),
-        hardMax: 1,
-      ),
-    ]);
+    return PreferenceBuilder(
+      preference: _tabSettings.detailsPropertyDisplayType,
+      builder: (context, DetailsPropertyDisplayType detailsPropertyDisplayType) {
+        late List<Widget> widgets;
+
+        switch (detailsPropertyDisplayType) {
+          case DetailsPropertyDisplayType.appearance:
+            final temperatureTypeProperties = _getPropertiesOnValidActors('TemperatureType');
+
+            widgets = [
+              UnrealDropdownSelector(
+                overrideName: AppLocalizations.of(context)!.propertyCCRType,
+                unrealProperties: _getPropertiesOnValidActors('Type'),
+              ),
+              UnrealMultiPropertyBuilder<String>(
+                properties: temperatureTypeProperties,
+                fallbackValue: AppLocalizations.of(context)!.mismatchedValuesLabel,
+                builder: (_, String? sharedValue, __) => UnrealDeltaSlider(
+                  overrideName: sharedValue,
+                  unrealProperties: _getPropertiesOnValidActors('Temperature'),
+                  buildLabel: (name) => UnrealDropdownText(unrealProperties: temperatureTypeProperties),
+                ),
+              ),
+              UnrealDeltaSlider(
+                unrealProperties: _getPropertiesOnValidActors('Tint'),
+              ),
+              UnrealDeltaSlider(
+                unrealProperties: _getPropertiesOnValidActors('Intensity'),
+              ),
+              UnrealDeltaSlider(
+                unrealProperties: _getPropertiesOnValidActors('Inner'),
+              ),
+              UnrealDeltaSlider(
+                unrealProperties: _getPropertiesOnValidActors('Outer'),
+              ),
+              UnrealDeltaSlider(
+                unrealProperties: _getPropertiesOnValidActors('Falloff'),
+                hardMax: 1,
+              ),
+            ];
+            break;
+
+          case DetailsPropertyDisplayType.orientation:
+            // CCR transform settings live on the "Root" SceneComponent instead of the actor itself
+            getPropertiesOnValidActorRoots(String name) => _getPropertiesOnValidActors(
+                  name,
+                  subObjectName: 'Root',
+                );
+
+            widgets = [
+              UnrealMultiPropertyBuilder<bool?>(
+                properties: [
+                  ...getPropertiesOnValidActorRoots('bAbsoluteLocation'),
+                  ...getPropertiesOnValidActorRoots('bAbsoluteRotation'),
+                  ...getPropertiesOnValidActorRoots('bAbsoluteScale'),
+                ],
+                fallbackValue: null,
+                builder: (context, sharedValue, _) => DropdownSelectorPropertyWrapper(
+                  propertyLabel: AppLocalizations.of(context)!.propertyTransformMode,
+                  child: DropdownSelector<bool?>(
+                    value: sharedValue,
+                    items: [false, true],
+                    hint: AppLocalizations.of(context)!.mismatchedValuesLabel,
+                    onChanged: (newValue) => _setTransformModeOnValidActors(newValue ?? false),
+                    makeItemName: (item) => switch (item) {
+                      true => AppLocalizations.of(context)!.transformModeAbsolute,
+                      false => AppLocalizations.of(context)!.transformModeRelative,
+                      null => AppLocalizations.of(context)!.mismatchedValuesLabel,
+                    },
+                  ),
+                ),
+              ),
+              UnrealDeltaSlider(
+                overrideName: AppLocalizations.of(context)!.propertyLocationX,
+                softMin: -100,
+                softMax: 100,
+                bAreMinMaxRelative: true,
+                unrealProperties: getPropertiesOnValidActorRoots('RelativeLocation.X'),
+              ),
+              UnrealDeltaSlider(
+                overrideName: AppLocalizations.of(context)!.propertyLocationY,
+                softMin: -100,
+                softMax: 100,
+                bAreMinMaxRelative: true,
+                unrealProperties: getPropertiesOnValidActorRoots('RelativeLocation.Y'),
+              ),
+              UnrealDeltaSlider(
+                overrideName: AppLocalizations.of(context)!.propertyLocationZ,
+                softMin: -100,
+                softMax: 100,
+                bAreMinMaxRelative: true,
+                unrealProperties: getPropertiesOnValidActorRoots('RelativeLocation.Z'),
+              ),
+              const PropertyDivider(padding: 14),
+              UnrealDeltaSlider(
+                overrideName: AppLocalizations.of(context)!.propertyRotationPitch,
+                softMin: -180,
+                softMax: 180,
+                unrealProperties: getPropertiesOnValidActorRoots('RelativeRotation.Pitch'),
+              ),
+              UnrealDeltaSlider(
+                overrideName: AppLocalizations.of(context)!.propertyRotationYaw,
+                softMin: -180,
+                softMax: 180,
+                unrealProperties: getPropertiesOnValidActorRoots('RelativeRotation.Yaw'),
+              ),
+              UnrealDeltaSlider(
+                overrideName: AppLocalizations.of(context)!.propertyRotationRoll,
+                softMin: -180,
+                softMax: 180,
+                unrealProperties: getPropertiesOnValidActorRoots('RelativeRotation.Roll'),
+              ),
+              const PropertyDivider(padding: 14),
+              UnrealDeltaSlider(
+                overrideName: AppLocalizations.of(context)!.propertyScaleX,
+                softMin: 0,
+                softMax: 10,
+                unrealProperties: getPropertiesOnValidActorRoots('RelativeScale3D.X'),
+              ),
+              UnrealDeltaSlider(
+                overrideName: AppLocalizations.of(context)!.propertyScaleY,
+                softMin: 0,
+                softMax: 10,
+                unrealProperties: getPropertiesOnValidActorRoots('RelativeScale3D.Y'),
+              ),
+              UnrealDeltaSlider(
+                overrideName: AppLocalizations.of(context)!.propertyScaleZ,
+                softMin: 0,
+                softMax: 10,
+                unrealProperties: getPropertiesOnValidActorRoots('RelativeScale3D.Z'),
+              ),
+            ];
+            break;
+        }
+
+        return Column(children: widgets);
+      },
+    );
   }
 
   /// Create property widgets for a CCW.
@@ -517,7 +629,7 @@ class _DetailsTabState extends State<DetailsTab> with GuardedRefreshState {
     final List<String> scalableActorClasses = [lightCardClassName];
     scalableActorClasses.addAll(colorCorrectWindowClassNames);
 
-    final List<String> scalableActorPaths = _getValidActorPathsOfClasses(scalableActorClasses);
+    final Iterable<String> scalableActorPaths = _getValidActorPathsOfClasses(scalableActorClasses);
 
     final bool bHasNonUvLightCards = scalableActorPaths.any(
       (actorPath) => _actorManager.getActorAtPath(actorPath)?.getPerClassData<LightCardActorData>()?.bIsUV != true,
@@ -590,14 +702,23 @@ class _DetailsTabState extends State<DetailsTab> with GuardedRefreshState {
   /// Get a list of properties with the given [name] for all selected actors that belong to a valid class.
   /// If [modifierFunction] is provided, it will be called for each [actorPath] and the [propertyName] of the property,
   /// and its return value will be used in place of [name].
+  /// If [typeNameOverride] is provided, it will be used instead of the default typename for the property.
+  /// If [subObjectName] is provided, it it will be appended to actor paths.
   List<UnrealProperty> _getPropertiesOnValidActors(
     String name, {
     String Function(String actorPath, String propertyName)? modifierFunction,
     String? typeNameOverride,
+    String? subObjectName,
   }) {
+    Iterable<String> paths = _getValidActorPaths();
+
+    if (subObjectName != null) {
+      paths = paths.map((path) => path + '.' + subObjectName);
+    }
+
     return _getPropertiesOnActors(
       name,
-      _getValidActorPaths(),
+      paths,
       modifierFunction: modifierFunction,
       typeNameOverride: typeNameOverride,
     );
@@ -615,7 +736,7 @@ class _DetailsTabState extends State<DetailsTab> with GuardedRefreshState {
   }
 
   /// Get a list of selected actor paths that we want to edit (i.e. have valid classes).
-  List<String> _getValidActorPaths() {
+  Iterable<String> _getValidActorPaths() {
     return _selectedActorSettings.selectedActors
         .getValue()
         .where((actorPath) => _isActorOfValidClass(actorPath))
@@ -623,7 +744,86 @@ class _DetailsTabState extends State<DetailsTab> with GuardedRefreshState {
   }
 
   /// Return all valid selected actors actors that are a member of at least one class in [classNames].
-  List<String> _getValidActorPathsOfClasses(List<String> classNames) => _getValidActorPaths()
+  Iterable<String> _getValidActorPathsOfClasses(List<String> classNames) => _getValidActorPaths()
       .where((path) => _actorManager.getActorAtPath(path)?.isAny(classNames) ?? false)
       .toList(growable: false);
+
+  /// Set the transform modes of all valid actors' root SceneComponents. If [bNewIsAbsolute] is true, they will be set
+  /// to Absolute mode; otherwise, they will be set to Relative.
+  void _setTransformModeOnValidActors(bool bNewIsAbsolute) async {
+    final transactionManager = Provider.of<UnrealTransactionManager>(context, listen: false);
+
+    // Start a transaction
+    if (!transactionManager.beginTransaction(AppLocalizations.of(context)!.transactionChangeTransformModes)) {
+      _log.warning('Failed to begin transaction for _setTransformModeOnValidActors');
+      return;
+    }
+
+    // Helper function to create a function call HTTP request
+    makeFunctionCall(String objectPath, String functionName, [bool generateTransaction = false, parameters = null]) =>
+        UnrealHttpRequest(
+          url: '/remote/object/call',
+          verb: 'PUT',
+          body: {
+            'objectPath': objectPath,
+            'functionName': functionName,
+            'parameters': parameters ?? {},
+            'generateTransaction': generateTransaction,
+          },
+        );
+
+    // Get the current world transform of each component
+    final componentPaths = _getValidActorPaths().map((path) => path + '.Root').toList(growable: false);
+    final worldTransformRequests = componentPaths.map((componentPath) => makeFunctionCall(
+          componentPath,
+          'GetWorldTransform',
+        ));
+
+    final connectionManager = Provider.of<EngineConnectionManager>(context, listen: false);
+    final UnrealHttpResponse worldTransformResponse =
+        await connectionManager.sendBatchedHttpRequest(worldTransformRequests);
+
+    if (worldTransformResponse.code != HttpResponseCode.ok) {
+      _log.warning('Failed to get world transforms for components in _setTransformModeOnValidActors');
+      transactionManager.endTransaction();
+      return;
+    }
+
+    if (worldTransformResponse.body.length != componentPaths.length) {
+      _log.warning('Wrong number of component world transforms received in _setTransformModeOnValidActors');
+      transactionManager.endTransaction();
+      return;
+    }
+
+    // Successfully received transforms. Now update each component's transform modes, then set world transform back
+    // so that the coordinates are correct in the new modes.
+    final List<UnrealHttpRequest> updateTransformRequests = [];
+
+    for (int componentIndex = 0; componentIndex < componentPaths.length; ++componentIndex) {
+      final String componentPath = componentPaths[componentIndex];
+      final worldTransform = worldTransformResponse.body?[componentIndex]?.body?['ReturnValue'];
+
+      if (worldTransform == null) {
+        _log.warning('No world transform received for $componentPath');
+        continue;
+      }
+
+      updateTransformRequests.add(makeFunctionCall(componentPath, 'SetAbsolute', true, {
+        'bNewAbsoluteLocation': bNewIsAbsolute,
+        'bNewAbsoluteRotation': bNewIsAbsolute,
+        'bNewAbsoluteScale': bNewIsAbsolute,
+      }));
+
+      updateTransformRequests.add(
+        makeFunctionCall(componentPath, 'SetWorldTransform', true, {
+          'NewTransform': worldTransform,
+        }),
+      );
+    }
+
+    connectionManager.sendBatchedHttpRequest(updateTransformRequests);
+
+    // Finally, close out the transaction
+    transactionManager.endTransaction();
+  }
 }
