[UnrealStage] Remove white bar under tab bar

diff --git a/Engine/Extras/VirtualProduction/EpicStageApp/lib/widgets/screens/main/toolbar/main_screen_tab_bar.dart b/Engine/Extras/VirtualProduction/EpicStageApp/lib/widgets/screens/main/toolbar/main_screen_tab_bar.dart
--- a/Engine/Extras/VirtualProduction/EpicStageApp/lib/widgets/screens/main/toolbar/main_screen_tab_bar.dart
+++ b/Engine/Extras/VirtualProduction/EpicStageApp/lib/widgets/screens/main/toolbar/main_screen_tab_bar.dart
@@ -26,6 +26,7 @@ class _MainScreenTabBarState extends State<MainScreenTabBar> {
       controller: widget.tabController,
       isScrollable: true, // This shrinks the tab buttons to fit their contents
       unselectedLabelColor: UnrealColors.white,
+      dividerHeight: 0,
       padding: EdgeInsets.zero,
       indicatorPadding: EdgeInsets.zero,
       labelPadding: const EdgeInsets.symmetric(horizontal: 4),
