[UnrealStage] Update "selected pins only" icon

diff --git a/Engine/Extras/VirtualProduction/EpicStageApp/lib/widgets/elements/lightcard_map.dart b/Engine/Extras/VirtualProduction/EpicStageApp/lib/widgets/elements/lightcard_map.dart
--- a/Engine/Extras/VirtualProduction/EpicStageApp/lib/widgets/elements/lightcard_map.dart
+++ b/Engine/Extras/VirtualProduction/EpicStageApp/lib/widgets/elements/lightcard_map.dart
@@ -1995,7 +1995,7 @@ class _ShowOnlySelectedPinsToggle extends StatelessWidget {
           tooltipMessage: bShowOnlySelectedPins
               ? AppLocalizations.of(context)!.stageMapShowOnlySelectedPinsOff
               : AppLocalizations.of(context)!.stageMapShowOnlySelectedPinsOn,
-          iconPath: 'assets/images/icons/solo.svg',
+          iconPath: 'assets/images/icons/show_selected_pins.svg',
           buttonSize: const Size(48, 36),
           bIsToggledOn: bShowOnlySelectedPins,
           onPressed: () => stageMapSettings.bShowOnlySelectedPins.setValue(!bShowOnlySelectedPins),
