[UnrealStage] Fix connection tiles connecting to wrong editor instance

diff --git a/Engine/Extras/VirtualProduction/EpicStageApp/lib/widgets/screens/connect/views/connection_tile.dart b/Engine/Extras/VirtualProduction/EpicStageApp/lib/widgets/screens/connect/views/connection_tile.dart
--- a/Engine/Extras/VirtualProduction/EpicStageApp/lib/widgets/screens/connect/views/connection_tile.dart
+++ b/Engine/Extras/VirtualProduction/EpicStageApp/lib/widgets/screens/connect/views/connection_tile.dart
@@ -31,15 +31,6 @@ class ConnectionTile extends StatefulWidget {
 }
 
 class _ConnectionTileState extends State<ConnectionTile> with ConnectMixin {
-  /// State instance for [widget.connectionData], for easy access across child widgets.
-  late final ConnectionData data;
-
-  @override
-  void initState() {
-    data = widget.connectionData;
-    super.initState();
-  }
-
   @override
   Widget build(BuildContext context) {
     late final Color fillColor;
@@ -56,7 +47,7 @@ class _ConnectionTileState extends State<ConnectionTile> with ConnectMixin {
     late final Color textColor = widget.bFilled ? Theme.of(context).colorScheme.onPrimary : UnrealColors.gray56;
 
     return InkWell(
-      onTap: () => connect(data),
+      onTap: () => connect(widget.connectionData),
       child: Container(
         width: 170,
         height: 166,
