[UnrealStage] Show warning if cine camera is invalid

diff --git a/Engine/Extras/VirtualProduction/EpicStageApp/lib/models/unreal_property_manager.dart b/Engine/Extras/VirtualProduction/EpicStageApp/lib/models/unreal_property_manager.dart
--- a/Engine/Extras/VirtualProduction/EpicStageApp/lib/models/unreal_property_manager.dart
+++ b/Engine/Extras/VirtualProduction/EpicStageApp/lib/models/unreal_property_manager.dart
@@ -816,10 +816,6 @@ class UnrealPropertyManager {
       previousValue: property.value,
       conversionMetadata: property.conversionMetadata,
     );
-    if (engineValue == null) {
-      _log.warning('Failed to read PropertyValue for change to $propertyLabel');
-      return;
-    }
 
     // Start at the time given by the server, then reapply our own input since then
     final List<_SequencedPropertyChangeRecord> recentTickChanges = property.recentTickChanges;
@@ -999,20 +995,24 @@ class _TrackedProperty {
   /// Underlying predicted value of the property. Use the [value] instead so callbacks are triggered.
   dynamic _value;
 
+  /// Whether we've received the first value from the engine.
+  bool _bIsReady = false;
+
   /// The current predicted value of the property.
   dynamic get value => _value;
 
   set value(dynamic newValue) {
-    bool bWasReady = _value != null;
     _value = newValue;
 
     // If we weren't ready before, then this is the first time the value was set. Callers are expecting to receive
     // confirmation that the property was exposed (including its value) rather than a callback containing the value,
     // so don't send callbacks yet.
-    if (bWasReady) {
+    if (_bIsReady) {
       for (PropertyChangeCallback callback in changeCallbacks) {
         callback(newValue);
       }
+    } else {
+      _bIsReady = true;
     }
   }
 }
diff --git a/Engine/Extras/VirtualProduction/EpicStageApp/lib/widgets/elements/actor_reference_viewer.dart b/Engine/Extras/VirtualProduction/EpicStageApp/lib/widgets/elements/actor_reference_viewer.dart
--- a/Engine/Extras/VirtualProduction/EpicStageApp/lib/widgets/elements/actor_reference_viewer.dart
+++ b/Engine/Extras/VirtualProduction/EpicStageApp/lib/widgets/elements/actor_reference_viewer.dart
@@ -130,6 +130,10 @@ class _UnrealActorDisplayState extends State<UnrealActorReferenceViewer>
       },
     ));
 
+    if (!mounted) {
+      return;
+    }
+
     if (response.code != HttpResponseCode.ok) {
       setState(() {
         _actorName = AppLocalizations.of(context)!.actorReferenceInvalidLabel;
diff --git a/Engine/Extras/VirtualProduction/EpicStageApp/lib/widgets/screens/main/tabs/camera_settings_tab.dart b/Engine/Extras/VirtualProduction/EpicStageApp/lib/widgets/screens/main/tabs/camera_settings_tab.dart
--- a/Engine/Extras/VirtualProduction/EpicStageApp/lib/widgets/screens/main/tabs/camera_settings_tab.dart
+++ b/Engine/Extras/VirtualProduction/EpicStageApp/lib/widgets/screens/main/tabs/camera_settings_tab.dart
@@ -27,6 +27,8 @@
 import '../../../elements/dropdown_button.dart';
 import '../../../elements/reset_mode_button.dart';
 import '../sidebar/outliner_panel.dart';
+import '../toolbar/settings/pages/settings_log_list.dart';
+import '../toolbar/settings/settings_dialog.dart';
 
 final _log = Logger('CameraTab');
 
@@ -265,19 +267,84 @@
 }
 
 /// The main controls of the camera settings tab.
-class _CameraSettingsMainControls extends StatelessWidget {
+class _CameraSettingsMainControls extends StatefulWidget {
   const _CameraSettingsMainControls({required this.selectedCamera});
 
   /// The camera selected for editing.
   final _CameraSettingsOutlinerEntryData? selectedCamera;
 
+  @override
+  State<StatefulWidget> createState() => _CameraSettingsMainControlsState();
+}
+
+class _CameraSettingsMainControlsState extends State<_CameraSettingsMainControls> {
+  /// The camera that holds the properties. Null if it's invalid or hasn't been retrieved yet.
+  UnrealObject? _camera;
+
+  /// The camera that holds the properties. Null if it's invalid or hasn't been retrieved yet.
+  UnrealObject? get camera => _camera;
+
+  /// True if the camera has been retrieved, but is invalid.
+  bool _bIsCameraInvalid = false;
+
+  /// True if the camera has been retrieved, but is invalid.
+  bool get bIsCameraInvalid => _bIsCameraInvalid;
+
+  /// The actor manager used to watch for changes to cine cameras.
+  late final UnrealActorManager _actorManager;
+
+  /// Class name for cine cameras
+  static const String _cineCameraClassName = '/Script/CinematicCamera.CineCameraActor';
+
   /// Height of the header area in each column.
-  final double _headerHeight = 52;
+  static const double _headerHeight = 52;
+
+  set camera(UnrealObject? value) {
+    if (!mounted || (value?.path == _camera?.path && !_bIsCameraInvalid)) {
+      return;
+    }
+
+    setState(() {
+      _camera = value;
+    });
+  }
+
+  /// Property controller used to listen to changes to the focus method.
+  UnrealPropertyController<String>? _cineCameraActorController;
 
   @override
+  void initState() {
+    super.initState();
+
+    _actorManager = Provider.of<UnrealActorManager>(context, listen: false);
+    _actorManager.watchClassName(_cineCameraClassName, _onCineCameraUpdate);
+
+    _updateCamera();
+  }
+
+  @override
   Widget build(BuildContext context) {
-    if (selectedCamera == null) {
-      return const SizedBox();
+    if (bIsCameraInvalid) {
+      return EmptyPlaceholder(
+        message: AppLocalizations.of(context)!.cameraTabInvalidCineCameraWarning(widget.selectedCamera!.name),
+        iconPath: 'packages/epic_common/assets/icons/alert_triangle_large_white.svg',
+        button: EpicWideButton(
+          text: AppLocalizations.of(context)!.cameraTabOpenApplicationLogButtonLabel,
+          iconPath: 'packages/epic_common/assets/icons/log.svg',
+          onPressed: () => SettingsDialog.show(SettingsLogList.route),
+        ),
+      );
     }
 
+    if (widget.selectedCamera == null || camera == null) {
+      return const Center(
+        child: SizedBox.square(
+          dimension: 80,
+          child: CircularProgressIndicator(
+            strokeWidth: 6,
+          ),
+        ),
+      );
+    }
+
     final settings = Provider.of<CameraSettingsTabSettings>(context);
@@ -320,10 +387,12 @@
                         preference: settings.leftColumnMode,
                         builder: (context, CameraSettingsLeftColumnMode mode) => switch (mode) {
                           CameraSettingsLeftColumnMode.camera => _CameraSettingsMainControlsCameraPropertyDisplay(
-                              selectedCamera: selectedCamera,
+                              selectedCamera: widget.selectedCamera,
+                              camera: camera,
                             ),
                           CameraSettingsLeftColumnMode.focus => _CameraSettingsMainControlsFocusPropertyDisplay(
-                              selectedCamera: selectedCamera,
+                              selectedCamera: widget.selectedCamera,
+                              camera: camera,
                             ),
                         },
                       ),
@@ -353,7 +422,7 @@
                     padding: EdgeInsets.only(left: UnrealTheme.cardMargin, right: UnrealTheme.cardMargin, top: 16),
                     color: Theme.of(context).colorScheme.surface,
                     child: _CameraSettingsMainControlsDepthOfFieldPropertyDisplay(
-                      selectedCamera: selectedCamera,
+                      selectedCamera: widget.selectedCamera,
                     ),
                   ),
                 ),
@@ -364,26 +433,142 @@
       ),
     );
   }
+
+  @override
+  void didUpdateWidget(covariant _CameraSettingsMainControls oldWidget) {
+    super.didUpdateWidget(oldWidget);
+
+    if (oldWidget.selectedCamera?.path != widget.selectedCamera?.path) {
+      _updateCamera();
+    }
+  }
+
+  @override
+  void dispose() {
+    _cineCameraActorController?.dispose();
+    _actorManager.stopWatchingClassName(_cineCameraClassName, _onCineCameraUpdate);
+
+    super.dispose();
+  }
+
+  /// Update the path to the [camera] based on the configured settings.
+  void _updateCamera() async {
+    _cineCameraActorController?.dispose();
+    _cineCameraActorController = null;
+
+    final String? cameraPath = widget.selectedCamera?.path;
+    if (cameraPath == null) {
+      return;
+    }
+
+    _cineCameraActorController = UnrealPropertyController(context);
+    _cineCameraActorController!.addListener(_onCineCameraChanged);
+    _cineCameraActorController!.trackAllProperties([
+      UnrealProperty(
+        objectPath: cameraPath,
+        propertyName: 'CameraSettings.ExternalCameraActor',
+      ),
+    ]);
+  }
+
+  /// Called when the current camera's [cineCameraPath] property value changes.
+  void _onCineCameraChanged() async {
+    if (!mounted) {
+      return;
+    }
+
+    final String? cineCameraPath = _cineCameraActorController?.getSingleSharedValue().value;
+
+    if (cineCameraPath == null || cineCameraPath.isEmpty) {
+      if (widget.selectedCamera != null && mounted) {
+        // No cine camera set, so use the camera selected in the app
+        _bIsCameraInvalid = false;
+        camera = UnrealObject(
+          path: widget.selectedCamera!.path,
+          name: widget.selectedCamera!.name,
+        );
+      }
+      return;
+    }
+
+    // Retrieve the camera component path and name of the cine camera actor
+    final connection = Provider.of<EngineConnectionManager>(context, listen: false);
+    final UnrealHttpResponse result = await connection.sendHttpRequest(
+      UnrealHttpRequest(
+        url: '/remote/object/property',
+        verb: 'PUT',
+        body: {
+          'objectPath': cineCameraPath,
+          'propertyName': 'CameraComponent',
+          'access': 'READ_ACCESS',
+        },
+      ),
+    );
+
+    if (!mounted) {
+      return;
+    }
+
+    // Retrieve the cine camera component path
+    if (result.code != HttpResponseCode.ok) {
+      _log.warning('Broken reference to Actor ID "$cineCameraPath" on camera "${widget.selectedCamera?.path}"');
+      _bIsCameraInvalid = true;
+      camera = null;
+      return;
+    }
+
+    final componentPath = result.body['CameraComponent'];
+    if (componentPath is! String) {
+      _log.warning('No camera component found on cine camera actor $cineCameraPath');
+      _bIsCameraInvalid = true;
+      camera = null;
+      return;
+    }
+
+    _bIsCameraInvalid = false;
+    camera = UnrealObject(path: componentPath, name: componentPath.split('.').last);
+  }
+
+  /// Called when there's an update about a cine camera
+  void _onCineCameraUpdate(ActorUpdateDetails details) {
+    final String? cineCameraPath = _cineCameraActorController?.getSingleSharedValue().value;
+
+    // If the cine camera has been re-added (e.g. by redo), try fetching it
+    if (_bIsCameraInvalid && details.addedActors.any((actor) => actor.path == cineCameraPath)) {
+      _updateCamera();
+      return;
+    }
+
+    // If the cine camera has been deleted, invalidate it
+    if (camera != null && details.deletedActors.any((actor) => actor.path == cineCameraPath)) {
+      _log.warning('Cine camera "${camera!.path}" deleted while still referenced by "${widget.selectedCamera?.path}"');
+      _bIsCameraInvalid = true;
+      camera = null;
+    }
+  }
 }
 
 /// Camera property controls in the Camera Settings tab's main controls pane.
-class _CameraSettingsMainControlsCameraPropertyDisplay extends _CineCameraSettingsPropertyDisplay {
-  const _CameraSettingsMainControlsCameraPropertyDisplay({required super.selectedCamera});
+class _CameraSettingsMainControlsCameraPropertyDisplay extends StatelessWidget with _CineCameraSettingsPropertyDisplay {
+  const _CameraSettingsMainControlsCameraPropertyDisplay({
+    required this.selectedCamera,
+    required this.camera,
+  });
 
   @override
-  State<StatefulWidget> createState() => _CameraSettingsMainControlsCameraPropertyDisplayState();
-}
+  final _CameraSettingsOutlinerEntryData? selectedCamera;
 
-class _CameraSettingsMainControlsCameraPropertyDisplayState
-    extends _CineCameraSettingsPropertyDisplayState<_CameraSettingsMainControlsCameraPropertyDisplay> {
   @override
+  final UnrealObject? camera;
+
+  @override
   Widget build(BuildContext context) {
     return Column(children: [
       UnrealActorReferenceViewer(
         unrealProperties: [
-          if (widget.selectedCamera != null)
+          if (selectedCamera != null)
             UnrealProperty(
-              objectPath: widget.selectedCamera!.path,
+              objectPath: selectedCamera!.path,
               propertyName: 'CameraSettings.ExternalCameraActor',
             ),
         ],
@@ -403,16 +588,24 @@
 }
 
 /// Focus property controls in the Camera Settings tab's main controls pane.
-class _CameraSettingsMainControlsFocusPropertyDisplay extends _CineCameraSettingsPropertyDisplay {
-  const _CameraSettingsMainControlsFocusPropertyDisplay({required super.selectedCamera});
+class _CameraSettingsMainControlsFocusPropertyDisplay extends StatefulWidget with _CineCameraSettingsPropertyDisplay {
+  const _CameraSettingsMainControlsFocusPropertyDisplay({
+    required this.selectedCamera,
+    required this.camera,
+  });
 
   @override
+  final _CameraSettingsOutlinerEntryData? selectedCamera;
+
+  @override
+  final UnrealObject? camera;
+
+  @override
   State<StatefulWidget> createState() => _CameraSettingsMainControlsFocusPropertyDisplayState();
 }
 
 class _CameraSettingsMainControlsFocusPropertyDisplayState
-    extends _CineCameraSettingsPropertyDisplayState<_CameraSettingsMainControlsFocusPropertyDisplay>
-    with GuardedRefreshState {
+    extends State<_CameraSettingsMainControlsFocusPropertyDisplay> with GuardedRefreshState {
   /// Property controller used to listen to changes to the focus method.
   UnrealPropertyController<String>? _focusMethodController;
 
@@ -431,5 +624,14 @@
   }
 
   @override
+  void didUpdateWidget(covariant _CameraSettingsMainControlsFocusPropertyDisplay oldWidget) {
+    super.didUpdateWidget(oldWidget);
+
+    if (oldWidget.camera != widget.camera) {
+      _updateFocusMethodTracking();
+    }
+  }
+
+  @override
   Widget build(BuildContext context) {
     final focusMode = _focusMethodController?.getSingleSharedValue().value;
@@ -440,7 +642,7 @@
       case 'Manual':
         propertyWidgets = [
           UnrealDeltaSlider(
-            unrealProperties: getCineCameraProperty('FocusSettings.ManualFocusDistance'),
+            unrealProperties: widget.getCineCameraProperty('FocusSettings.ManualFocusDistance'),
             hardMin: 1.5,
             softMax: 100000,
             minMaxBehaviour: PropertyMinMaxBehaviour.clamp,
@@ -452,22 +654,22 @@
       case 'Tracking':
         propertyWidgets = [
           UnrealActorReferenceViewer(
-            unrealProperties: getCineCameraProperty('FocusSettings.TrackingFocusSettings.ActorToTrack'),
+            unrealProperties: widget.getCineCameraProperty('FocusSettings.TrackingFocusSettings.ActorToTrack'),
           ),
           UnrealDeltaSlider(
-            unrealProperties: getCineCameraProperty('FocusSettings.TrackingFocusSettings.RelativeOffset.X'),
+            unrealProperties: widget.getCineCameraProperty('FocusSettings.TrackingFocusSettings.RelativeOffset.X'),
             overrideName: AppLocalizations.of(context)!.cameraTabFocusRelativeOffsetXLabel,
             softMin: -10000,
             softMax: 10000,
           ),
           UnrealDeltaSlider(
-            unrealProperties: getCineCameraProperty('FocusSettings.TrackingFocusSettings.RelativeOffset.Y'),
+            unrealProperties: widget.getCineCameraProperty('FocusSettings.TrackingFocusSettings.RelativeOffset.Y'),
             overrideName: AppLocalizations.of(context)!.cameraTabFocusRelativeOffsetYLabel,
             softMin: -10000,
             softMax: 10000,
           ),
           UnrealDeltaSlider(
-            unrealProperties: getCineCameraProperty('FocusSettings.TrackingFocusSettings.RelativeOffset.Z'),
+            unrealProperties: widget.getCineCameraProperty('FocusSettings.TrackingFocusSettings.RelativeOffset.Z'),
             overrideName: AppLocalizations.of(context)!.cameraTabFocusRelativeOffsetZLabel,
             softMin: -10000,
             softMax: 10000,
@@ -483,29 +685,24 @@
 
     return Column(children: [
       UnrealDropdownSelector(
-        unrealProperties: getCineCameraProperty('FocusSettings.FocusMethod'),
+        unrealProperties: widget.getCineCameraProperty('FocusSettings.FocusMethod'),
       ),
       ...propertyWidgets,
     ]);
   }
 
-  @override
-  void onCameraChanged() {
-    _updateFocusMethodTracking();
-  }
-
   /// Create the list of widgets used in both "enabled" focus modes (i.e. Manual/Tracking).
   List<Widget> _makeGenericFocusPropertyWidgets() => [
         UnrealBooleanDropdownSelector(
-          unrealProperties: getCineCameraProperty('FocusSettings.bSmoothFocusChanges'),
+          unrealProperties: widget.getCineCameraProperty('FocusSettings.bSmoothFocusChanges'),
         ),
         UnrealDeltaSlider(
-          unrealProperties: getCineCameraProperty('FocusSettings.FocusSmoothingInterpSpeed'),
+          unrealProperties: widget.getCineCameraProperty('FocusSettings.FocusSmoothingInterpSpeed'),
           hardMin: 0,
           softMax: 10,
         ),
         UnrealDeltaSlider(
-          unrealProperties: getCineCameraProperty('FocusSettings.FocusOffset'),
+          unrealProperties: widget.getCineCameraProperty('FocusSettings.FocusOffset'),
           softMin: -1000,
           softMax: 1000,
         ),
@@ -515,7 +712,7 @@
   void _updateFocusMethodTracking() {
     _focusMethodController = UnrealPropertyController(context);
 
-    final UnrealProperty? property = getCineCameraProperty('FocusSettings.FocusMethod').firstOrNull;
+    final UnrealProperty? property = widget.getCineCameraProperty('FocusSettings.FocusMethod').firstOrNull;
 
     if (property == null) {
       return;
@@ -641,62 +838,14 @@
   }
 }
 
-/// Base class for widgets that display properties in the Camera Settings tab, and may pull these properties from either
+/// Mixin for widgets that display properties in the Camera Settings tab, and may pull these properties from either
 /// the selected camera itself or its associated cine camera.
-abstract class _CineCameraSettingsPropertyDisplay extends StatefulWidget {
-  const _CineCameraSettingsPropertyDisplay({required this.selectedCamera});
-
+mixin _CineCameraSettingsPropertyDisplay {
   /// The camera selected for editing.
-  final _CameraSettingsOutlinerEntryData? selectedCamera;
-}
+  _CameraSettingsOutlinerEntryData? get selectedCamera;
 
-/// Base state class for widgets that display properties in the Camera Settings tab, and may pull these properties from
-/// either the selected camera itself or its associated cine camera.
-abstract class _CineCameraSettingsPropertyDisplayState<T extends _CineCameraSettingsPropertyDisplay> extends State<T> {
-  /// The camera that holds the properties. Null if it hasn't been retrieved yet.
-  /// Use [camera] to access this value, which will automatically call [onCameraChanged] when appropriate.
-  UnrealObject? _camera;
+  /// The camera actor from which to retrieve properties, which may be the cine camera or the selected camera.
+  UnrealObject? get camera;
 
-  /// The camera that holds the properties. Null if it hasn't been retrieved yet.
-  UnrealObject? get camera => _camera;
-
-  set camera(UnrealObject? value) {
-    if (!mounted || value?.path == _camera?.path) {
-      return;
-    }
-
-    setState(() {
-      _camera = value;
-    });
-
-    onCameraChanged();
-  }
-
-  /// Property controller used to listen to changes to the focus method.
-  UnrealPropertyController<String>? _cineCameraActorController;
-
-  @override
-  void initState() {
-    super.initState();
-
-    _updateCamera();
-  }
-
-  @override
-  void didUpdateWidget(covariant T oldWidget) {
-    super.didUpdateWidget(oldWidget);
-
-    if (oldWidget.selectedCamera?.path != widget.selectedCamera?.path) {
-      _updateCamera();
-    }
-  }
-
-  @override
-  void dispose() {
-    _cineCameraActorController?.dispose();
-
-    super.dispose();
-  }
-
   /// Get a property by [name] on the cine camera, or if it isn't present, on the selected camera.
   /// Returns the property in a list for use with delta widgets.
@@ -713,83 +862,6 @@
       )
     ];
   }
-
-  /// Called when the camera whose properties we want to modify changes.
-  @protected
-  void onCameraChanged() {}
-
-  /// Update the path to the [camera] based on the configured settings.
-  void _updateCamera() async {
-    _cineCameraActorController?.dispose();
-
-    final String? cameraPath = widget.selectedCamera?.path;
-    if (cameraPath == null) {
-      return;
-    }
-
-    _cineCameraActorController = UnrealPropertyController(context);
-    _cineCameraActorController!.addListener(_onCineCameraChanged);
-    _cineCameraActorController!.trackAllProperties([
-      UnrealProperty(
-        objectPath: cameraPath,
-        propertyName: 'CameraSettings.ExternalCameraActor',
-      ),
-    ]);
-  }
-
-  /// Called when the current camera's [cineCameraPath] property value changes.
-  void _onCineCameraChanged() async {
-    if (!mounted) {
-      return;
-    }
-
-    final String? cineCameraPath = _cineCameraActorController?.getSingleSharedValue().value;
-
-    if (cineCameraPath == null || cineCameraPath.isEmpty) {
-      if (widget.selectedCamera != null && mounted) {
-        // No cine camera set, so use the camera selected in the app
-        camera = UnrealObject(
-          path: widget.selectedCamera!.path,
-          name: widget.selectedCamera!.name,
-        );
-      }
-      return;
-    }
-
-    // Retrieve the camera component path and name of the cine camera actor
-    final connection = Provider.of<EngineConnectionManager>(context, listen: false);
-    final UnrealHttpResponse result = await connection.sendHttpRequest(
-      UnrealHttpRequest(
-        url: '/remote/object/property',
-        verb: 'PUT',
-        body: {
-          'objectPath': cineCameraPath,
-          'propertyName': 'CameraComponent',
-          'access': 'READ_ACCESS',
-        },
-      ),
-    );
-
-    if (!mounted) {
-      return;
-    }
-
-    // Retrieve the cine camera component path
-    if (result.code != HttpResponseCode.ok) {
-      _log.warning('Failed to retrieve cine camera component: ${result.body}');
-      camera = null;
-      return;
-    }
-
-    final componentPath = result.body['CameraComponent'];
-    if (componentPath is! String) {
-      _log.warning('No camera component found on cine camera actor $cineCameraPath');
-      camera = null;
-      return;
-    }
-
-    camera = UnrealObject(path: componentPath, name: componentPath.split('.').last);
-  }
 }
 
 /// The outliner panel for the color grading tab, showing cameras that can be targeted.
diff --git a/Engine/Extras/VirtualProduction/EpicStageApp/lib/widgets/screens/main/toolbar/settings/settings_dialog.dart b/Engine/Extras/VirtualProduction/EpicStageApp/lib/widgets/screens/main/toolbar/settings/settings_dialog.dart
--- a/Engine/Extras/VirtualProduction/EpicStageApp/lib/widgets/screens/main/toolbar/settings/settings_dialog.dart
+++ b/Engine/Extras/VirtualProduction/EpicStageApp/lib/widgets/screens/main/toolbar/settings/settings_dialog.dart
@@ -13,11 +13,21 @@ import 'pages/settings_root_actor_picker.dart';
 
 /// Dialog containing the app's settings menu.
 class SettingsDialog extends StatefulWidget {
-  const SettingsDialog({Key? key}) : super(key: key);
+  const SettingsDialog({
+    Key? key,
+    this.initialRoute,
+  }) : super(key: key);
 
-  static void show() {
+  /// The initial route to show.
+  final String? initialRoute;
+
+  /// Show the settings dialog.
+  /// If [initialRoute] is provided, open to that route instead of the main settings page.
+  static void show([String? initialRoute]) {
     final route = GenericModalDialogRoute(
-      builder: (_) => SettingsDialog(),
+      builder: (_) => SettingsDialog(
+        initialRoute: initialRoute,
+      ),
     );
 
     Navigator.of(rootNavigatorKey.currentContext!, rootNavigator: true).push(route);
@@ -46,6 +56,7 @@ class _SettingsDialogState extends State<SettingsDialog> with RouteAware {
         child: IntrinsicHeight(
           child: Navigator(
             key: _innerNavigatorKey,
+            initialRoute: widget.initialRoute,
             onGenerateRoute: (final RouteSettings settings) {
               late final Widget page;
 
