[UnrealVCam] Update app to support Tentacle SDK 4.0.2

diff --git a/Engine/Extras/VirtualProduction/FlutterLiveLinkVCAM/IOS/Runner/Tentacle/API/TentacleApi.g.swift b/Engine/Extras/VirtualProduction/FlutterLiveLinkVCAM/IOS/Runner/Tentacle/API/TentacleApi.g.swift
--- a/Engine/Extras/VirtualProduction/FlutterLiveLinkVCAM/IOS/Runner/Tentacle/API/TentacleApi.g.swift
+++ b/Engine/Extras/VirtualProduction/FlutterLiveLinkVCAM/IOS/Runner/Tentacle/API/TentacleApi.g.swift
@@ -48,6 +48,7 @@ enum TentacleProductId: Int {
   case generic = 0
   case syncE = 1
   case trackE = 2
+  case timebar = 3
 }
 
 /// Information about a Tentacle device. This only includes data which is expected to be requested infrequently, so it
