[UnrealVCam] [UnrealStage] [FlutterCommon] Update dependencies

diff --git a/Engine/Extras/FlutterCommon/pubspec.yaml b/Engine/Extras/FlutterCommon/pubspec.yaml
--- a/Engine/Extras/FlutterCommon/pubspec.yaml
+++ b/Engine/Extras/FlutterCommon/pubspec.yaml
@@ -18,17 +18,17 @@ dependencies:
   intl: any
   logging: ^1.2.0
   path_provider: ^2.0.11
-  package_info_plus: ^5.0.1
+  package_info_plus: ^8.1.1
   provider: ^6.0.2
   streaming_shared_preferences: ^2.0.0
   uuid: ^4.2.1
   archive: ^3.2.0
-  share_plus: ^7.2.1
+  share_plus: ^10.1.2
 
 dev_dependencies:
   flutter_test:
     sdk: flutter
-  flutter_lints: ^3.0.1
+  flutter_lints: ^5.0.0
 
 flutter:
   generate: true
diff --git a/Engine/Extras/VirtualProduction/EpicStageApp/IOS/Runner/AppDelegate.swift b/Engine/Extras/VirtualProduction/EpicStageApp/IOS/Runner/AppDelegate.swift
--- a/Engine/Extras/VirtualProduction/EpicStageApp/IOS/Runner/AppDelegate.swift
+++ b/Engine/Extras/VirtualProduction/EpicStageApp/IOS/Runner/AppDelegate.swift
@@ -3,7 +3,7 @@
 import UIKit
 import Flutter
 
-@UIApplicationMain
+@main
 @objc class AppDelegate: FlutterAppDelegate {
   override func application(
     _ application: UIApplication,
diff --git a/Engine/Extras/VirtualProduction/EpicStageApp/pubspec.yaml b/Engine/Extras/VirtualProduction/EpicStageApp/pubspec.yaml
--- a/Engine/Extras/VirtualProduction/EpicStageApp/pubspec.yaml
+++ b/Engine/Extras/VirtualProduction/EpicStageApp/pubspec.yaml
@@ -36,7 +36,7 @@ dependencies:
 
   uuid: ^4.2.1
   archive: ^3.2.0
-  web_socket_channel: ^2.1.0
+  web_socket_channel: ^3.0.1
   provider: ^6.0.2
   flutter_svg: ^2.0.9
   streaming_shared_preferences: ^2.0.0
@@ -45,7 +45,7 @@ dependencies:
   intl: any
   path_provider: ^2.0.11
   open_filex: ^4.3.1
-  share_plus: ^7.2.1
+  share_plus: ^10.1.2
   flutter_localizations:
     sdk: flutter
   arb_generator: ^0.0.3
@@ -62,9 +62,9 @@ dev_dependencies:
   # activated in the `analysis_options.yaml` file located at the root of your
   # package. See that file for information about deactivating specific lint
   # rules and activating additional ones.
-  flutter_lints: ^3.0.1
+  flutter_lints: ^5.0.0
 
-  flutter_launcher_icons: ^0.13.1
+  flutter_launcher_icons: ^0.14.1
 
 # For information on the generic Dart part of this file, see the
 # following page: https://dart.dev/tools/pub/pubspec
diff --git a/Engine/Extras/VirtualProduction/FlutterLiveLinkVCAM/IOS/Runner/AppDelegate.swift b/Engine/Extras/VirtualProduction/FlutterLiveLinkVCAM/IOS/Runner/AppDelegate.swift
--- a/Engine/Extras/VirtualProduction/FlutterLiveLinkVCAM/IOS/Runner/AppDelegate.swift
+++ b/Engine/Extras/VirtualProduction/FlutterLiveLinkVCAM/IOS/Runner/AppDelegate.swift
@@ -5,7 +5,7 @@ import Flutter
 import WebRTC
 import SpriteKit
 
-@UIApplicationMain
+@main
 @objc class AppDelegate: FlutterAppDelegate {
   var peerConnectionApi: FlutterRtcPeerConnectionApi?
   var dataChannelApi: FlutterRtcDataChannelApi?
diff --git a/Engine/Extras/VirtualProduction/FlutterLiveLinkVCAM/pubspec.yaml b/Engine/Extras/VirtualProduction/FlutterLiveLinkVCAM/pubspec.yaml
--- a/Engine/Extras/VirtualProduction/FlutterLiveLinkVCAM/pubspec.yaml
+++ b/Engine/Extras/VirtualProduction/FlutterLiveLinkVCAM/pubspec.yaml
@@ -18,34 +18,34 @@ dependencies:
   epic_common:
     path: ../../FlutterCommon
 
-  device_info_plus: ^9.1.1
+  device_info_plus: ^11.1.0
   ffi: ^2.1.0
-  flutter_widget_from_html_core: ^0.14.11
+  flutter_widget_from_html_core: ^0.15.2
   logging: ^1.2.0
   ntp: ^2.0.0
-  package_info_plus: ^5.0.1
+  package_info_plus: ^8.1.0
   permission_handler: ^11.1.0
   provider: ^6.1.1
   system_clock: ^2.0.0
-  web_socket_channel: ^2.1.0
+  web_socket_channel: ^3.0.1
   streaming_shared_preferences: ^2.0.0
-  share_plus: ^7.2.1
+  share_plus: ^10.1.2
 
 dev_dependencies:
   flutter_test:
     sdk: flutter
 
-  ffigen: ^10.0.0
-  pigeon: ^14.0.1
+  ffigen: ^16.0.0
+  pigeon: ^22.6.3
   url_launcher: ^6.2.2
-  flutter_launcher_icons: ^0.13.1
+  flutter_launcher_icons: ^0.14.1
 
   # The "flutter_lints" package below contains a set of recommended lints to
   # encourage good coding practices. The lint set provided by the package is
   # activated in the `analysis_options.yaml` file located at the root of your
   # package. See that file for information about deactivating specific lint
   # rules and activating additional ones.
-  flutter_lints: ^3.0.1
+  flutter_lints: ^5.0.0
 
 # For information on the generic Dart part of this file, see the
 # following page: https://dart.dev/tools/pub/pubspec
