[UnrealVCAM] Fix iOS build after Tentacle SDK update

diff --git a/Engine/Extras/VirtualProduction/FlutterLiveLinkVCAM/IOS/Runner/Tentacle/API/FlutterTentacleApi.swift b/Engine/Extras/VirtualProduction/FlutterLiveLinkVCAM/IOS/Runner/Tentacle/API/FlutterTentacleApi.swift
--- a/Engine/Extras/VirtualProduction/FlutterLiveLinkVCAM/IOS/Runner/Tentacle/API/FlutterTentacleApi.swift
+++ b/Engine/Extras/VirtualProduction/FlutterLiveLinkVCAM/IOS/Runner/Tentacle/API/FlutterTentacleApi.swift
@@ -17,7 +17,7 @@ class FlutterTentacleApi: FlutterPluginApi<TentacleFlutterApi> {
     TentacleHostApiSetup.setUp(binaryMessenger: binaryMessenger, api: self)
     FlutterTentacleApi.instance = self
     
-    bluetooth = TentacleBluetoothController()
+    bluetooth = TentacleBluetoothController(DispatchQueue.main)
     bluetooth.delegate = self
   }
   
