Update the third party notices to the plural form. Add version bump to pubspec.yaml

diff --git a/Engine/Extras/VirtualProduction/EpicStageApp/lib/widgets/screens/main/toolbar/settings/pages/settings_dialog_main.dart b/Engine/Extras/VirtualProduction/EpicStageApp/lib/widgets/screens/main/toolbar/settings/pages/settings_dialog_main.dart
--- a/Engine/Extras/VirtualProduction/EpicStageApp/lib/widgets/screens/main/toolbar/settings/pages/settings_dialog_main.dart
+++ b/Engine/Extras/VirtualProduction/EpicStageApp/lib/widgets/screens/main/toolbar/settings/pages/settings_dialog_main.dart
@@ -84,7 +84,7 @@ class SettingsDialogMain extends StatelessWidget {
             onTap: () => launchUrl(Uri.parse(_privacyPolicyUri)),
           ),
           SettingsMenuItem(
-            title: EpicCommonLocalizations.of(context)!.settingsThirdPartyNoticeTitle,
+            title: EpicCommonLocalizations.of(context)!.settingsThirdPartyNoticesTitle,
             iconPath: 'packages/epic_common/assets/icons/info.svg',
             onTap: () => Navigator.of(context).pushNamed(SettingsThirdPartyNoticeList.route),
           ),
diff --git a/Engine/Extras/VirtualProduction/EpicStageApp/pubspec.yaml b/Engine/Extras/VirtualProduction/EpicStageApp/pubspec.yaml
--- a/Engine/Extras/VirtualProduction/EpicStageApp/pubspec.yaml
+++ b/Engine/Extras/VirtualProduction/EpicStageApp/pubspec.yaml
@@ -15,7 +15,7 @@ publish_to: 'none' # Remove this line if you wish to publish to pub.dev
 # In iOS, build-name is used as CFBundleShortVersionString while build-number used as CFBundleVersion.
 # Read more about iOS versioning at
 # https://developer.apple.com/library/archive/documentation/General/Reference/InfoPlistKeyReference/Articles/CoreFoundationKeys.html
-version: 1.0.3+4
+version: 1.3.1+43843285
 
 environment:
   sdk: ">=3.0.6 <4.0.0"
