Version bump and clean-up of README.md.  Also fixed IOS side to remove flutter framework.

diff --git a/Engine/Extras/VirtualProduction/FlutterLiveLinkVCAM/README.md b/Engine/Extras/VirtualProduction/FlutterLiveLinkVCAM/README.md
--- a/Engine/Extras/VirtualProduction/FlutterLiveLinkVCAM/README.md
+++ b/Engine/Extras/VirtualProduction/FlutterLiveLinkVCAM/README.md
@@ -4,7 +4,7 @@ A rebuild of the Live Link VCAM app for Flutter, enabling cross-platform deploym
 
 ## Tentacle
 
-This app depends on the `FlutterTentacle` package (included in the Unreal Engine repo under `Engine/Extras/FlutterTentacle`), which in turn depends on the Tentacle SDK. The Tentacle SDK is not distributed to engine licensees. See the `FlutterTentacle` README for more info.
+This app depends on the `FlutterTentacle` package (included in the Unreal Engine repo under `Engine/Extras/Flutter/FlutterTentacle`), which in turn depends on the Tentacle SDK. The Tentacle SDK is not distributed to engine licensees. See the `FlutterTentacle` README for more info.
 
 ### IOS Note
 
@@ -21,5 +21,5 @@ To setup signing certificates, follow the instructions outlined in this article
 
 - [Deploying Flutter Apps to The PlayStore](https:[REDACTED_EMAIL]/deploying-flutter-apps-to-the-playstore-1bd0cce0d15c)
 
-Add your `android-key.tks` to `android\app` folder and `key.properities` to `android\`.  The gradle files are already setup
-to read these files when app bundling
\ No newline at end of file
+Add your `android-key.tks` to `android\app` folder and `key.properties` to `android\`.  The gradle files are already setup
+to read these files when app bundling
diff --git a/Engine/Extras/VirtualProduction/FlutterLiveLinkVCAM/pubspec.yaml b/Engine/Extras/VirtualProduction/FlutterLiveLinkVCAM/pubspec.yaml
--- a/Engine/Extras/VirtualProduction/FlutterLiveLinkVCAM/pubspec.yaml
+++ b/Engine/Extras/VirtualProduction/FlutterLiveLinkVCAM/pubspec.yaml
@@ -2,7 +2,7 @@ name: live_link_vcam
 description: A rebuild of the Live Link VCAM app for Flutter, enabling cross-platform deployment.
 publish_to: 'none'
 
-version: 2.0.0+32867938
+version: 2.2.1+43891908
 
 environment:
   sdk: '>=3.1.0 <4.0.0'
