Version string bunp and bump minimum version for Android SDK.

#rnx

#Synced-CL 47671294

diff --git a/Engine/Extras/VirtualProduction/FlutterLiveLinkVCAM/pubspec.yaml b/Engine/Extras/VirtualProduction/FlutterLiveLinkVCAM/pubspec.yaml
--- a/Engine/Extras/VirtualProduction/FlutterLiveLinkVCAM/pubspec.yaml
+++ b/Engine/Extras/VirtualProduction/FlutterLiveLinkVCAM/pubspec.yaml
@@ -2,7 +2,7 @@ name: live_link_vcam
 description: A rebuild of the Live Link VCAM app for Flutter, enabling cross-platform deployment.
 publish_to: 'none'
 
-version: 2.2.1+43891908
+version: 2.3.0+47388791
 
 environment:
   sdk: '>=3.1.0 <4.0.0'
